JSON API improvements for numbers

Added int32/int64 overrides for the WriteValue functions that didn't take a key so that they matched the variations of the functions which did take a key.
Added an int64 override for TryGetNumber to match the int32/uint32 versions.

diff --git a/Engine/Source/Runtime/Json/Private/Dom/JsonValue.cpp b/Engine/Source/Runtime/Json/Private/Dom/JsonValue.cpp
--- a/Engine/Source/Runtime/Json/Private/Dom/JsonValue.cpp
+++ b/Engine/Source/Runtime/Json/Private/Dom/JsonValue.cpp
@@ -109,6 +109,28 @@ bool FJsonValue::TryGetNumber( uint32& OutNumber ) const
 }
 
 
+bool FJsonValue::TryGetNumber( int64& OutNumber ) const
+{
+	double Double;
+
+	if (TryGetNumber(Double) && (Double >= INT64_MIN) && (Double <= INT64_MAX))
+	{
+		if (Double >= 0.0)
+		{
+			OutNumber = (int64)(Double + 0.5);
+		}
+		else
+		{
+			OutNumber = (int64)(Double - 0.5);
+		}
+		
+		return true;
+	}
+
+	return false;
+}
+
+
 //static 
 bool FJsonValue::CompareEqual( const FJsonValue& Lhs, const FJsonValue& Rhs )
 {
diff --git a/Engine/Source/Runtime/Json/Public/Dom/JsonValue.h b/Engine/Source/Runtime/Json/Public/Dom/JsonValue.h
--- a/Engine/Source/Runtime/Json/Public/Dom/JsonValue.h
+++ b/Engine/Source/Runtime/Json/Public/Dom/JsonValue.h
@@ -38,6 +38,9 @@ public:
 	/** Tries to convert this value to an integer, returning false if not possible */
 	bool TryGetNumber(uint32& OutNumber) const;
 
+	/** Tries to convert this value to an integer, returning false if not possible */
+	bool TryGetNumber(int64& OutNumber) const;
+
 	/** Tries to convert this value to a string, returning false if not possible */
 	virtual bool TryGetString(FString& OutString) const { return false; }
 
diff --git a/Engine/Source/Runtime/Json/Public/Serialization/JsonWriter.h b/Engine/Source/Runtime/Json/Public/Serialization/JsonWriter.h
--- a/Engine/Source/Runtime/Json/Public/Serialization/JsonWriter.h
+++ b/Engine/Source/Runtime/Json/Public/Serialization/JsonWriter.h
@@ -238,6 +238,44 @@ public:
 		PreviousTokenWritten = EJsonToken::Number;
 	}
 
+	void WriteValue( const int32 Value )
+	{
+		check(CanWriteValueWithoutIdentifier());
+		WriteCommaIfNeeded();
+
+		if ( PreviousTokenWritten != EJsonToken::Number && PreviousTokenWritten != EJsonToken::SquareOpen )
+		{
+			PrintPolicy::WriteLineTerminator(Stream);
+			PrintPolicy::WriteTabs(Stream, IndentLevel);
+		}
+		else
+		{
+			PrintPolicy::WriteSpace( Stream );
+		}
+
+		WriteIntegerValue( Value );
+		PreviousTokenWritten = EJsonToken::Number;
+	}
+
+	void WriteValue( const int64 Value )
+	{
+		check(CanWriteValueWithoutIdentifier());
+		WriteCommaIfNeeded();
+
+		if ( PreviousTokenWritten != EJsonToken::Number && PreviousTokenWritten != EJsonToken::SquareOpen )
+		{
+			PrintPolicy::WriteLineTerminator(Stream);
+			PrintPolicy::WriteTabs(Stream, IndentLevel);
+		}
+		else
+		{
+			PrintPolicy::WriteSpace( Stream );
+		}
+
+		WriteIntegerValue( Value );
+		PreviousTokenWritten = EJsonToken::Number;
+	}
+
 	void WriteValue( const FString& Value )
 	{
 		check(CanWriteValueWithoutIdentifier());
@@ -248,6 +286,11 @@ public:
 		PreviousTokenWritten = EJsonToken::String;
 	}
 
+	void WriteValue( const TCHAR* Value )
+	{
+		WriteValue(FString(Value));
+	}
+
 	void WriteNull()
 	{
 		check(CanWriteValueWithoutIdentifier());
