Copying //Tasks/UE4/Dev-Online-Plugins-v3 to //UE4/Dev-Main (Source: //Tasks/UE4/Dev-Online-Plugins-v3 @ 3047294)

diff --git a/Engine/Source/Runtime/Engine/Private/PlayerMuteList.cpp b/Engine/Source/Runtime/Engine/Private/PlayerMuteList.cpp
--- a/Engine/Source/Runtime/Engine/Private/PlayerMuteList.cpp
+++ b/Engine/Source/Runtime/Engine/Private/PlayerMuteList.cpp
@@ -1,14 +1,14 @@
 // Copyright 1998-2016 Epic Games, Inc. All Rights Reserved.
 
 #include "EnginePrivate.h"
-#include "OnlineSubsystemUtils.h"
+#include "Net/OnlineEngineInterface.h"
 #include "GameFramework/GameSession.h"
 #include "GameFramework/PlayerState.h"
 
 static inline void AddIdToMuteList(TArray< TSharedRef<const FUniqueNetId> >& MuteList, const TSharedPtr<const FUniqueNetId>& UniqueIdToAdd)
 {
-	FUniqueNetIdMatcher UniqueIdToAddMatch(*UniqueIdToAdd);
-	if (MuteList.IndexOfByPredicate(UniqueIdToAddMatch) == INDEX_NONE)
+	auto UniqueIdToAddPred = [&UniqueIdToAdd](TSharedRef<const FUniqueNetId> Other) { return UniqueIdToAdd.IsValid() && *UniqueIdToAdd == *Other; };
+	if (MuteList.IndexOfByPredicate(UniqueIdToAddPred) == INDEX_NONE)
 	{
 		MuteList.Add(UniqueIdToAdd.ToSharedRef());
 	}
@@ -16,8 +16,8 @@ static inline void AddIdToMuteList(TArray< TSharedRef<const FUniqueNetId> >& Mut
 
 static inline void RemoveIdFromMuteList(TArray< TSharedRef<const FUniqueNetId> >& MuteList, const TSharedPtr<const FUniqueNetId>& UniqueIdToRemove)
 {
-	FUniqueNetIdMatcher UniqueIdToRemoveMatch(*UniqueIdToRemove);
-	int32 RemoveIndex = MuteList.IndexOfByPredicate(UniqueIdToRemoveMatch);
+	auto UniqueIdToRemovePred = [&UniqueIdToRemove](TSharedRef<const FUniqueNetId> Other) { return UniqueIdToRemove.IsValid() && *UniqueIdToRemove == *Other; };
+	int32 RemoveIndex = MuteList.IndexOfByPredicate(UniqueIdToRemovePred);
 	if (RemoveIndex != INDEX_NONE)
 	{
 		MuteList.RemoveAtSwap(RemoveIndex);
@@ -64,29 +64,30 @@ void FPlayerMuteList::ServerUnmutePlayer(APlayerController* OwningPC, const FUni
 	APlayerController* OtherPC = GetPlayerControllerFromNetId(World, *PlayerIdToUnmute);
 	if (OtherPC != NULL)
 	{
-		FUniqueNetIdMatcher PlayerIdToUnmuteMatch(*PlayerIdToUnmute);
-		FUniqueNetIdMatcher OwningPlayerIdMatch(*OwningPC->PlayerState->UniqueId);
+		FUniqueNetIdRepl& OwningPlayerId = OwningPC->PlayerState->UniqueId;
+		auto PlayerIdToUnmutePred = [&PlayerIdToUnmute](TSharedRef<const FUniqueNetId> Other) { return PlayerIdToUnmute.IsValid() && *PlayerIdToUnmute == *Other; };
+		auto OwningPlayerIdPred = [&OwningPlayerId](TSharedRef<const FUniqueNetId> Other) { return *OwningPlayerId == *Other; };
 
 		// Make sure this player isn't muted for gameplay reasons
-		if (GameplayVoiceMuteList.IndexOfByPredicate(PlayerIdToUnmuteMatch) == INDEX_NONE &&
+		if (GameplayVoiceMuteList.IndexOfByPredicate(PlayerIdToUnmutePred) == INDEX_NONE &&
 			// And make sure they didn't mute us
-			OtherPC->MuteList.VoiceMuteList.IndexOfByPredicate(OwningPlayerIdMatch) == INDEX_NONE)
+			OtherPC->MuteList.VoiceMuteList.IndexOfByPredicate(OwningPlayerIdPred) == INDEX_NONE)
 		{
 			OwningPC->ClientUnmutePlayer(UnmuteId);
 		}
 
 		// If the other player doesn't have this player muted
-		if (OtherPC->MuteList.VoiceMuteList.IndexOfByPredicate(OwningPlayerIdMatch) == INDEX_NONE &&
-			OtherPC->MuteList.GameplayVoiceMuteList.IndexOfByPredicate(OwningPlayerIdMatch) == INDEX_NONE)
+		if (OtherPC->MuteList.VoiceMuteList.IndexOfByPredicate(OwningPlayerIdPred) == INDEX_NONE &&
+			OtherPC->MuteList.GameplayVoiceMuteList.IndexOfByPredicate(OwningPlayerIdPred) == INDEX_NONE)
 		{
 			// Remove them from the packet filter list
 			RemoveIdFromMuteList(VoicePacketFilter, PlayerIdToUnmute);
 
 			// If found, remove so packets flow to that client too
-			RemoveIdFromMuteList(OtherPC->MuteList.VoicePacketFilter, OwningPC->PlayerState->UniqueId.GetUniqueNetId());
+			RemoveIdFromMuteList(OtherPC->MuteList.VoicePacketFilter, OwningPlayerId.GetUniqueNetId());
 
 			// Tell the other PC to unmute this one
-			OtherPC->ClientUnmutePlayer(OwningPC->PlayerState->UniqueId);
+			OtherPC->ClientUnmutePlayer(OwningPlayerId);
 		}
 	}
 }
@@ -103,12 +104,9 @@ void FPlayerMuteList::ClientMutePlayer(APlayerController* OwningPC, const FUniqu
 	if (LP != NULL)
 	{
 		UWorld* World = OwningPC->GetWorld();
-		IOnlineVoicePtr VoiceInt = Online::GetVoiceInterface(World);
-		if (VoiceInt.IsValid())
-		{
-			// Have the voice subsystem mute this player
-			VoiceInt->MuteRemoteTalker(LP->GetControllerId(), *PlayerIdToMute, false);
-		}
+	
+		// Have the voice subsystem mute this player
+		UOnlineEngineInterface::Get()->MuteRemoteTalker(World, LP->GetControllerId(), *PlayerIdToMute, false);
 	}
 }
 
@@ -124,12 +122,9 @@ void FPlayerMuteList::ClientUnmutePlayer(APlayerController* OwningPC, const FUni
 	if (LP != NULL)
 	{
 		UWorld* World = OwningPC->GetWorld();
-		IOnlineVoicePtr VoiceInt = Online::GetVoiceInterface(World);
-		if (VoiceInt.IsValid())
-		{
-			// Have the voice subsystem mute this player
-			VoiceInt->UnmuteRemoteTalker(LP->GetControllerId(), *PlayerIdToUnmute, false);
-		}
+
+		// Have the voice subsystem mute this player
+		UOnlineEngineInterface::Get()->UnmuteRemoteTalker(World, LP->GetControllerId(), *PlayerIdToUnmute, false);
 	}
 }
 
@@ -160,13 +155,14 @@ void FPlayerMuteList::GameplayUnmutePlayer(APlayerController* OwningPC, const FU
 	APlayerController* OtherPC = GetPlayerControllerFromNetId(World, *PlayerIdToUnmute);
 	if (OtherPC != NULL)
 	{
-		FUniqueNetIdMatcher PlayerIdToUnmuteMatch(*PlayerIdToUnmute);
-		FUniqueNetIdMatcher OwningPlayerIdMatch(*OwningPC->PlayerState->UniqueId);
+		FUniqueNetIdRepl& OwningPlayerId = OwningPC->PlayerState->UniqueId;
+		auto PlayerIdToUnmutePred = [&PlayerIdToUnmute](TSharedRef<const FUniqueNetId> Other) { return PlayerIdToUnmute.IsValid() && *PlayerIdToUnmute == *Other; };
+		auto OwningPlayerIdPred = [&OwningPlayerId](TSharedRef<const FUniqueNetId> Other) { return *OwningPlayerId == *Other; };
 
 		// Make sure this player isn't explicitly muted
-		if (VoiceMuteList.IndexOfByPredicate(PlayerIdToUnmuteMatch) == INDEX_NONE &&
+		if (VoiceMuteList.IndexOfByPredicate(PlayerIdToUnmutePred) == INDEX_NONE &&
 			// And make sure they didn't mute us
-			OtherPC->MuteList.VoiceMuteList.IndexOfByPredicate(OwningPlayerIdMatch) == INDEX_NONE)
+			OtherPC->MuteList.VoiceMuteList.IndexOfByPredicate(OwningPlayerIdPred) == INDEX_NONE)
 		{
 			RemoveIdFromMuteList(VoicePacketFilter, PlayerIdToUnmute);
 
@@ -178,8 +174,8 @@ void FPlayerMuteList::GameplayUnmutePlayer(APlayerController* OwningPC, const FU
 
 bool FPlayerMuteList::IsPlayerMuted(const FUniqueNetId& PlayerId)
 {
-	FUniqueNetIdMatcher PlayerIdMatch(PlayerId);
-	return VoicePacketFilter.IndexOfByPredicate(PlayerIdMatch) != INDEX_NONE;
+	auto PlayerIdToUnmutePred = [&PlayerId](TSharedRef<const FUniqueNetId> Other) { return PlayerId == *Other; };
+	return VoicePacketFilter.IndexOfByPredicate(PlayerIdToUnmutePred) != INDEX_NONE;
 }
 
 FString DumpMutelistState(UWorld* World)
diff --git a/Engine/Source/Runtime/Engine/Private/PlayerState.cpp b/Engine/Source/Runtime/Engine/Private/PlayerState.cpp
--- a/Engine/Source/Runtime/Engine/Private/PlayerState.cpp
+++ b/Engine/Source/Runtime/Engine/Private/PlayerState.cpp
@@ -7,7 +7,7 @@
 #include "EnginePrivate.h"
 #include "GameFramework/EngineMessage.h"
 #include "Net/UnrealNetwork.h"
-#include "OnlineSubsystemUtils.h"
+#include "Net/OnlineEngineInterface.h"
 #include "GameFramework/GameState.h"
 #include "GameFramework/PlayerState.h"
 
@@ -271,12 +271,11 @@ void APlayerState::RegisterPlayerWithSession(bool bWasFromInvite)
 {
 	if (GetNetMode() != NM_Standalone)
 	{
-		IOnlineSessionPtr SessionInt = Online::GetSessionInterface(GetWorld());
-		if (SessionInt.IsValid() && UniqueId.IsValid())
+		if (UniqueId.IsValid()) // May not be valid if this is was created via DebugCreatePlayer
 		{
 			// Register the player as part of the session
 			const APlayerState* PlayerState = GetDefault<APlayerState>();
-			SessionInt->RegisterPlayer(PlayerState->SessionName, *UniqueId, bWasFromInvite);
+			UOnlineEngineInterface::Get()->RegisterPlayer(GetWorld(), PlayerState->SessionName, *UniqueId, bWasFromInvite);
 		}
 	}
 }
@@ -288,12 +287,7 @@ void APlayerState::UnregisterPlayerWithSession()
 		const APlayerState* PlayerState = GetDefault<APlayerState>();
 		if (PlayerState->SessionName != NAME_None)
 		{
-			IOnlineSessionPtr SessionInt = Online::GetSessionInterface(GetWorld());
-			if (SessionInt.IsValid())
-			{
-				// Remove the player from the session
-				SessionInt->UnregisterPlayer(PlayerState->SessionName, *UniqueId);
-			}
+			UOnlineEngineInterface::Get()->UnregisterPlayer(GetWorld(), PlayerState->SessionName, *UniqueId);
 		}
 	}
 }
diff --git a/Engine/Source/Runtime/Engine/Private/UnrealEngine.cpp b/Engine/Source/Runtime/Engine/Private/UnrealEngine.cpp
--- a/Engine/Source/Runtime/Engine/Private/UnrealEngine.cpp
+++ b/Engine/Source/Runtime/Engine/Private/UnrealEngine.cpp
@@ -9,6 +9,7 @@
 #include "Engine/LevelStreamingPersistent.h"
 #include "Engine/ObjectReferencer.h"
 #include "Net/UnrealNetwork.h"
+#include "Net/OnlineEngineInterface.h"
 #include "Engine/Console.h"
 #include "VisualLogger/VisualLogger.h"
 #include "FileManagerGeneric.h"
@@ -24,8 +25,6 @@
 #include "Networking.h"
 #include "ProfilingHelpers.h"
 #include "ImageWrapper.h"
-#include "OnlineSubsystem.h"
-#include "OnlineExternalUIInterface.h"
 #include "EngineAnalytics.h"
 #include "Runtime/Analytics/Analytics/Public/Interfaces/IAnalyticsProvider.h"
 #if WITH_EDITOR
@@ -951,32 +950,9 @@ void UEngine::Init(IEngineLoop* InEngineLoop)
 	UE_LOG(LogInit, Log, TEXT("Texture streaming: %s"), IStreamingManager::Get().IsTextureStreamingEnabled() ? TEXT("Enabled") : TEXT("Disabled") );
 
 	// Initialize the online subsystem as early as possible
-	IOnlineSubsystem* SubSystem = IOnlineSubsystem::IsLoaded() ? IOnlineSubsystem::Get() : nullptr;
-	if (SubSystem != nullptr)
-	{
-		IOnlineExternalUIPtr ExternalUI = SubSystem->GetExternalUIInterface();
-		if(ExternalUI.IsValid())
-		{
-			FOnExternalUIChangeDelegate OnExternalUIChangeDelegate;
-			OnExternalUIChangeDelegate.BindUObject(this, &UEngine::OnExternalUIChange);
-
-			ExternalUI->AddOnExternalUIChangeDelegate_Handle(OnExternalUIChangeDelegate);
-		}
-	}
-	// Initialize the platform online subsystem as early as possible also
-	IOnlineSubsystem* SubSystemConsole = IOnlineSubsystem::GetByPlatform();
-	if (SubSystemConsole != nullptr &&
-		SubSystem != SubSystemConsole)
-	{
-		IOnlineExternalUIPtr ExternalUI = SubSystemConsole->GetExternalUIInterface();
-		if (ExternalUI.IsValid())
-		{
-			FOnExternalUIChangeDelegate OnExternalUIChangeDelegate;
-			OnExternalUIChangeDelegate.BindUObject(this, &UEngine::OnExternalUIChange);
-
-			ExternalUI->AddOnExternalUIChangeDelegate_Handle(OnExternalUIChangeDelegate);
-		}
-	}
+	FOnlineExternalUIChanged OnExternalUIChangeDelegate;
+	OnExternalUIChangeDelegate.BindUObject(this, &UEngine::OnExternalUIChange);
+	UOnlineEngineInterface::Get()->BindToExternalUIOpening(OnExternalUIChangeDelegate);
 
 	// Initialise buffer visualization system data
 	GetBufferVisualizationData().Initialize();
diff --git a/Engine/Source/Runtime/Engine/Private/VoiceChannel.cpp b/Engine/Source/Runtime/Engine/Private/VoiceChannel.cpp
--- a/Engine/Source/Runtime/Engine/Private/VoiceChannel.cpp
+++ b/Engine/Source/Runtime/Engine/Private/VoiceChannel.cpp
@@ -6,8 +6,7 @@
 
 #include "EnginePrivate.h"
 #include "Net/UnrealNetwork.h"
-#include "OnlineSubsystemUtils.h"
-#include "OnlineSubsystemTypes.h"
+#include "Net/OnlineEngineInterface.h"
 #include "Net/NetworkProfiler.h"
 #include "Engine/VoiceChannel.h"
 
@@ -29,26 +28,22 @@ void UVoiceChannel::ReceivedBunch(FInBunch& Bunch)
 {
 	if (Connection->Driver && Connection->Driver->World)
 	{
-		IOnlineVoicePtr VoiceInt = Online::GetVoiceInterface(Connection->Driver->World);
-		if (VoiceInt.IsValid())
+		while (!Bunch.AtEnd())
 		{
-			while (!Bunch.AtEnd())
+			// Give the data to the local voice processing
+			TSharedPtr<FVoicePacket> VoicePacket = UOnlineEngineInterface::Get()->SerializeRemotePacket(Connection->Driver->World, Bunch);
+			if (VoicePacket.IsValid())
 			{
-				// Give the data to the local voice processing
-				TSharedPtr<FVoicePacket> VoicePacket = VoiceInt->SerializeRemotePacket(Bunch);
-				if (VoicePacket.IsValid())
+				if (Connection->Driver->ServerConnection == NULL)
 				{
-					if (Connection->Driver->ServerConnection == NULL)
-					{
-						// Possibly replicate the data to other clients
-						Connection->Driver->ReplicateVoicePacket(VoicePacket, Connection);
-					}
+					// Possibly replicate the data to other clients
+					Connection->Driver->ReplicateVoicePacket(VoicePacket, Connection);
+				}
 #if STATS
-					// Increment the number of voice packets we've received
-					Connection->Driver->VoicePacketsRecv++;
-					Connection->Driver->VoiceBytesRecv += VoicePacket->GetBufferSize();
+				// Increment the number of voice packets we've received
+				Connection->Driver->VoicePacketsRecv++;
+				Connection->Driver->VoiceBytesRecv += VoicePacket->GetBufferSize();
 #endif
-				}
 			}
 		}
 	}
diff --git a/Engine/Source/Runtime/GameLiveStreaming/GameLiveStreaming.Build.cs b/Engine/Source/Runtime/GameLiveStreaming/GameLiveStreaming.Build.cs
--- a/Engine/Source/Runtime/GameLiveStreaming/GameLiveStreaming.Build.cs
+++ b/Engine/Source/Runtime/GameLiveStreaming/GameLiveStreaming.Build.cs
@@ -23,7 +23,6 @@ namespace UnrealBuildTool.Rules
 					"RenderCore",
 					"ShaderCore",
 					"RHI",
-					"OnlineSubsystemUtils"	// @todo livestream: We need to move the useful features into engine and eliminate this module
 				}
 			);
 		}
diff --git a/Engine/Source/Runtime/HTTPChunkInstaller/HTTPChunkInstaller.Build.cs b/Engine/Source/Runtime/HTTPChunkInstaller/HTTPChunkInstaller.Build.cs
--- a/Engine/Source/Runtime/HTTPChunkInstaller/HTTPChunkInstaller.Build.cs
+++ b/Engine/Source/Runtime/HTTPChunkInstaller/HTTPChunkInstaller.Build.cs
@@ -11,8 +11,6 @@ public class HTTPChunkInstaller : ModuleRules
         PrivateDependencyModuleNames.Add("PakFile");
 		PrivateDependencyModuleNames.Add("HTTP");
         PrivateDependencyModuleNames.Add("BuildPatchServices");
-        PrivateDependencyModuleNames.Add("OnlineSubsystem");
-        PrivateDependencyModuleNames.Add("OnlineSubsystemUtils");
 
 		PrecompileForTargets = PrecompileTargetsType.None;
 	}
diff --git a/Engine/Source/Runtime/Json/Public/Json.h b/Engine/Source/Runtime/Json/Public/Json.h
--- a/Engine/Source/Runtime/Json/Public/Json.h
+++ b/Engine/Source/Runtime/Json/Public/Json.h
@@ -26,3 +26,4 @@ JSON_API DECLARE_LOG_CATEGORY_EXTERN(LogJson, Log, All);
 #include "JsonReader.h"
 #include "JsonWriter.h"
 #include "JsonSerializer.h"
+#include "JsonSerializerMacros.h"
diff --git a/Engine/Source/Runtime/Json/Public/Serialization/JsonSerializerMacros.h b/Engine/Source/Runtime/Json/Public/Serialization/JsonSerializerMacros.h
new file mode 100644
--- /dev/null
+++ b/Engine/Source/Runtime/Json/Public/Serialization/JsonSerializerMacros.h
@@ -0,0 +1,679 @@
+// Copyright 1998-2016 Epic Games, Inc. All Rights Reserved.
+
+#pragma once
+#include "Json.h"
+
+/**
+ * Macros used to generate a serialization function for a class derived from FJsonSerializable
+ */
+#define BEGIN_JSON_SERIALIZER \
+	virtual void Serialize(FJsonSerializerBase& Serializer, bool bFlatObject) override \
+	{ \
+		if (!bFlatObject) { Serializer.StartObject(); }
+
+#define END_JSON_SERIALIZER \
+		if (!bFlatObject) { Serializer.EndObject(); } \
+	}
+
+#define JSON_SERIALIZE(JsonName, JsonValue) \
+		Serializer.Serialize(TEXT(JsonName), JsonValue)
+
+#define JSON_SERIALIZE_ARRAY(JsonName, JsonArray) \
+	    Serializer.SerializeArray(TEXT(JsonName), JsonArray)
+
+#define JSON_SERIALIZE_MAP(JsonName, JsonMap) \
+		Serializer.SerializeMap(TEXT(JsonName), JsonMap)
+
+#define JSON_SERIALIZE_SERIALIZABLE(JsonName, JsonValue) \
+		JsonValue.Serialize(Serializer, false)
+
+#define JSON_SERIALIZE_ARRAY_SERIALIZABLE(JsonName, JsonArray, ElementType) \
+		if (Serializer.IsLoading()) \
+		{ \
+			if (Serializer.GetObject()->HasTypedField<EJson::Array>(JsonName)) \
+			{ \
+				for (auto It = Serializer.GetObject()->GetArrayField(JsonName).CreateConstIterator(); It; ++It) \
+				{ \
+					ElementType* Obj = new(JsonArray) ElementType(); \
+					Obj->FromJson((*It)->AsObject()); \
+				} \
+			} \
+		} \
+		else \
+		{ \
+			Serializer.StartArray(JsonName); \
+			for (auto It = JsonArray.CreateIterator(); It; ++It) \
+			{ \
+				It->Serialize(Serializer, false); \
+			} \
+			Serializer.EndArray(); \
+		}
+
+#define JSON_SERIALIZE_MAP_SERIALIZABLE(JsonName, JsonMap, ElementType) \
+		if (Serializer.IsLoading()) \
+		{ \
+			if (Serializer.GetObject()->HasTypedField<EJson::Object>(JsonName)) \
+			{ \
+				TSharedPtr<FJsonObject> JsonObj = Serializer.GetObject()->GetObjectField(JsonName); \
+				for (auto MapIt = JsonObj->Values.CreateConstIterator(); MapIt; ++MapIt) \
+				{ \
+					ElementType NewEntry; \
+					NewEntry.FromJson(MapIt.Value()->AsObject()); \
+					JsonMap.Add(MapIt.Key(), NewEntry); \
+				} \
+			} \
+		} \
+		else \
+		{ \
+			Serializer.StartObject(JsonName); \
+			for (auto It = JsonMap.CreateIterator(); It; ++It) \
+			{ \
+				Serializer.StartObject(It.Key()); \
+				It.Value().Serialize(Serializer, true); \
+				Serializer.EndObject(); \
+			} \
+			Serializer.EndObject(); \
+		}
+
+#define JSON_SERIALIZE_OBJECT_SERIALIZABLE(JsonName, JsonSerializableObject) \
+		/* Process the JsonName field differently because it is an object */ \
+		if (Serializer.IsLoading()) \
+		{ \
+			/* Read in the value from the JsonName field */ \
+			if (Serializer.GetObject()->HasTypedField<EJson::Object>(JsonName)) \
+			{ \
+				TSharedPtr<FJsonObject> JsonObj = Serializer.GetObject()->GetObjectField(JsonName); \
+				if (JsonObj.IsValid()) \
+				{ \
+					JsonSerializableObject.FromJson(JsonObj); \
+				} \
+			} \
+		} \
+		else \
+		{ \
+			/* Write the value to the Name field */ \
+			Serializer.WriteIdentifierPrefix(TEXT(JsonName)); \
+			JsonSerializableObject.Serialize(Serializer, true); \
+		}
+
+/** Array of string data */
+typedef TArray<FString> FJsonSerializableArray;
+
+/** Maps a key to a value */
+typedef TMap<FString, FString> FJsonSerializableKeyValueMap;
+typedef TMap<FString, int32> FJsonSerializableKeyValueMapInt;
+
+/**
+ * Base interface used to serialize to/from JSON. Hides the fact there are separate read/write classes
+ */
+struct FJsonSerializerBase
+{
+	virtual bool IsLoading() const = 0;
+	virtual bool IsSaving() const = 0;
+	virtual void StartObject() = 0;
+	virtual void StartObject(const FString& Name) = 0;
+	virtual void EndObject() = 0;
+	virtual void StartArray() = 0;
+	virtual void StartArray(const FString& Name) = 0;
+	virtual void EndArray() = 0;
+	virtual void Serialize(const TCHAR* Name, int32& Value) = 0;
+	virtual void Serialize(const TCHAR* Name, uint32& Value) = 0;
+	virtual void Serialize(const TCHAR* Name, bool& Value) = 0;
+	virtual void Serialize(const TCHAR* Name, FString& Value) = 0;
+	virtual void Serialize(const TCHAR* Name, FText& Value) = 0;
+	virtual void Serialize(const TCHAR* Name, float& Value) = 0;
+	virtual void Serialize(const TCHAR* Name, double& Value) = 0;
+	virtual void Serialize(const TCHAR* Name, FDateTime& Value) = 0;
+	virtual void SerializeArray(FJsonSerializableArray& Array) = 0;
+	virtual void SerializeArray(const TCHAR* Name, FJsonSerializableArray& Value) = 0;
+	virtual void SerializeMap(const TCHAR* Name, FJsonSerializableKeyValueMap& Map) = 0;
+	virtual void SerializeMap(const TCHAR* Name, FJsonSerializableKeyValueMapInt& Map) = 0;
+	virtual TSharedPtr<FJsonObject> GetObject() = 0;
+	virtual void WriteIdentifierPrefix(const TCHAR* Name) = 0;
+};
+
+/**
+ * Implements the abstract serializer interface hiding the underlying writer object
+ */
+template <class CharType = TCHAR, class PrintPolicy = TPrettyJsonPrintPolicy<CharType> >
+class FJsonSerializerWriter :
+	public FJsonSerializerBase
+{
+	/** The object to write the JSON output to */
+	TSharedRef<TJsonWriter<CharType, PrintPolicy> > JsonWriter;
+
+public:
+
+	/**
+	 * Initializes the writer object
+	 *
+	 * @param InJsonWriter the object to write the JSON output to
+	 */
+	FJsonSerializerWriter(TSharedRef<TJsonWriter<CharType, PrintPolicy> > InJsonWriter) :
+		JsonWriter(InJsonWriter)
+	{
+	}
+
+    /** Is the JSON being read from */
+	virtual bool IsLoading() const override { return false; }
+    /** Is the JSON being written to */
+	virtual bool IsSaving() const override { return true; }
+	/** Access to the root object */
+	virtual TSharedPtr<FJsonObject> GetObject() override { return TSharedPtr<FJsonObject>(); }
+
+	/**
+	 * Starts a new object "{"
+	 */
+	virtual void StartObject() override
+	{
+		JsonWriter->WriteObjectStart();
+	}
+
+	/**
+	 * Starts a new object "{"
+	 */
+	virtual void StartObject(const FString& Name) override
+	{
+		JsonWriter->WriteObjectStart(Name);
+	}
+	/**
+	 * Completes the definition of an object "}"
+	 */
+	virtual void EndObject() override
+	{
+		JsonWriter->WriteObjectEnd();
+	}
+
+	virtual void StartArray() override
+	{
+		JsonWriter->WriteArrayStart();
+	}
+
+	virtual void StartArray(const FString& Name) override
+	{
+		JsonWriter->WriteArrayStart(Name);
+	}
+
+	virtual void EndArray() override
+	{
+		JsonWriter->WriteArrayEnd();
+	}
+	/**
+	 * Writes the field name and the corresponding value to the JSON data
+	 *
+	 * @param Name the field name to write out
+	 * @param Value the value to write out
+	 */
+	virtual void Serialize(const TCHAR* Name, int32& Value) override
+	{
+		JsonWriter->WriteValue(Name, Value);
+	}
+	/**
+	 * Writes the field name and the corresponding value to the JSON data
+	 *
+	 * @param Name the field name to write out
+	 * @param Value the value to write out
+	 */
+	virtual void Serialize(const TCHAR* Name, uint32& Value) override
+	{
+		JsonWriter->WriteValue(Name, static_cast<int64>(Value));
+	}
+	/**
+	 * Writes the field name and the corresponding value to the JSON data
+	 *
+	 * @param Name the field name to write out
+	 * @param Value the value to write out
+	 */
+	virtual void Serialize(const TCHAR* Name, bool& Value) override
+	{
+		JsonWriter->WriteValue(Name, Value);
+	}
+	/**
+	 * Writes the field name and the corresponding value to the JSON data
+	 *
+	 * @param Name the field name to write out
+	 * @param Value the value to write out
+	 */
+	virtual void Serialize(const TCHAR* Name, FString& Value) override
+	{
+		JsonWriter->WriteValue(Name, Value);
+	}
+	/**
+	 * Writes the field name and the corresponding value to the JSON data
+	 *
+	 * @param Name the field name to write out
+	 * @param Value the value to write out
+	 */
+	virtual void Serialize(const TCHAR* Name, FText& Value) override
+	{
+		JsonWriter->WriteValue(Name, Value.ToString());
+	}
+	/**
+	 * Writes the field name and the corresponding value to the JSON data
+	 *
+	 * @param Name the field name to write out
+	 * @param Value the value to write out
+	 */
+	virtual void Serialize(const TCHAR* Name, float& Value) override
+	{
+		JsonWriter->WriteValue(Name, Value);
+	}
+	/**
+	* Writes the field name and the corresponding value to the JSON data
+	*
+	* @param Name the field name to write out
+	* @param Value the value to write out
+	*/
+	virtual void Serialize(const TCHAR* Name, double& Value) override
+	{
+		JsonWriter->WriteValue(Name, Value);
+	}
+	/**
+	* Writes the field name and the corresponding value to the JSON data
+	*
+	* @param Name the field name to write out
+	* @param Value the value to write out
+	*/
+	virtual void Serialize(const TCHAR* Name, FDateTime& Value) override
+	{
+		if (Value.GetTicks() > 0)
+		{
+			JsonWriter->WriteValue(Name, Value.ToIso8601());
+		}
+	}
+	/**
+	 * Serializes an array of values
+	 *
+	 * @param Name the name of the property to serialize
+	 * @param Array the array to serialize
+	 */
+	virtual void SerializeArray(FJsonSerializableArray& Array) override
+	{
+		JsonWriter->WriteArrayStart();
+		// Iterate all of values
+		for (FJsonSerializableArray::TIterator ArrayIt(Array); ArrayIt; ++ArrayIt)
+		{
+			JsonWriter->WriteValue(*ArrayIt);
+		}
+		JsonWriter->WriteArrayEnd();
+	}
+	/**
+	 * Serializes an array of values with an identifier
+	 *
+	 * @param Name the name of the property to serialize
+	 * @param Array the array to serialize
+	 */
+	virtual void SerializeArray(const TCHAR* Name, FJsonSerializableArray& Array) override
+	{
+		JsonWriter->WriteArrayStart(Name);
+		// Iterate all of values
+		for (FJsonSerializableArray::TIterator ArrayIt(Array); ArrayIt; ++ArrayIt)
+		{
+			JsonWriter->WriteValue(*ArrayIt);
+		}
+		JsonWriter->WriteArrayEnd();
+	}
+	/**
+	 * Serializes the keys & values for map
+	 *
+	 * @param Name the name of the property to serialize
+	 * @param Map the map to serialize
+	 */
+	virtual void SerializeMap(const TCHAR* Name, FJsonSerializableKeyValueMap& Map) override
+	{
+		JsonWriter->WriteObjectStart(Name);
+		// Iterate all of the keys and their values
+		for (FJsonSerializableKeyValueMap::TIterator KeyValueIt(Map); KeyValueIt; ++KeyValueIt)
+		{
+			Serialize(*(KeyValueIt.Key()), KeyValueIt.Value());
+		}
+		JsonWriter->WriteObjectEnd();
+	}
+	/**
+	 * Serializes the keys & values for map
+	 *
+	 * @param Name the name of the property to serialize
+	 * @param Map the map to serialize
+	 */
+	virtual void SerializeMap(const TCHAR* Name, FJsonSerializableKeyValueMapInt& Map) override
+	{
+		JsonWriter->WriteObjectStart(Name);
+		// Iterate all of the keys and their values
+		for (FJsonSerializableKeyValueMapInt::TIterator KeyValueIt(Map); KeyValueIt; ++KeyValueIt)
+		{
+			Serialize(*(KeyValueIt.Key()), KeyValueIt.Value());
+		}
+		JsonWriter->WriteObjectEnd();
+	}
+
+	virtual void WriteIdentifierPrefix(const TCHAR* Name)
+	{
+		JsonWriter->WriteIdentifierPrefix(Name);
+	}
+};
+
+/**
+ * Implements the abstract serializer interface hiding the underlying reader object
+ */
+class FJsonSerializerReader :
+	public FJsonSerializerBase
+{
+	/** The object that holds the parsed JSON data */
+	TSharedPtr<FJsonObject> JsonObject;
+
+public:
+	/**
+	 * Inits the base JSON object that is being read from
+	 *
+	 * @param InJsonObject the JSON object to serialize from
+	 */
+	FJsonSerializerReader(TSharedPtr<FJsonObject> InJsonObject) :
+		JsonObject(InJsonObject)
+	{
+	}
+
+    /** Is the JSON being read from */
+	virtual bool IsLoading() const override { return true; }
+    /** Is the JSON being written to */
+	virtual bool IsSaving() const override { return false; }
+	/** Access to the root Json object being read */
+	virtual TSharedPtr<FJsonObject> GetObject() override { return JsonObject; }
+
+	/** Ignored */
+	virtual void StartObject() override
+	{
+		// Empty on purpose
+	}
+	/** Ignored */
+	virtual void StartObject(const FString& Name) override
+	{
+		// Empty on purpose
+	}
+	/** Ignored */
+	virtual void EndObject() override
+	{
+		// Empty on purpose
+	}
+	/** Ignored */
+	virtual void StartArray() override
+	{
+		// Empty on purpose
+	}
+	/** Ignored */
+	virtual void StartArray(const FString& Name) override
+	{
+		// Empty on purpose
+	}
+	/** Ignored */
+	virtual void EndArray() override
+	{
+		// Empty on purpose
+	}
+	/**
+	 * If the underlying json object has the field, it is read into the value
+	 *
+	 * @param Name the name of the field to read
+	 * @param Value the out value to read the data into
+	 */
+	virtual void Serialize(const TCHAR* Name, int32& Value) override
+	{
+		if (JsonObject->HasTypedField<EJson::Number>(Name))
+		{
+			JsonObject->TryGetNumberField(Name, Value);
+		}
+	}
+	/**
+	 * If the underlying json object has the field, it is read into the value
+	 *
+	 * @param Name the name of the field to read
+	 * @param Value the out value to read the data into
+	 */
+	virtual void Serialize(const TCHAR* Name, uint32& Value) override
+	{
+		if (JsonObject->HasTypedField<EJson::Number>(Name))
+		{
+			JsonObject->TryGetNumberField(Name, Value);
+		}
+	}
+	/**
+	 * If the underlying json object has the field, it is read into the value
+	 *
+	 * @param Name the name of the field to read
+	 * @param Value the out value to read the data into
+	 */
+	virtual void Serialize(const TCHAR* Name, bool& Value) override
+	{
+		if (JsonObject->HasTypedField<EJson::Boolean>(Name))
+		{
+			Value = JsonObject->GetBoolField(Name);
+		}
+	}
+	/**
+	 * If the underlying json object has the field, it is read into the value
+	 *
+	 * @param Name the name of the field to read
+	 * @param Value the out value to read the data into
+	 */
+	virtual void Serialize(const TCHAR* Name, FString& Value) override
+	{
+		if (JsonObject->HasTypedField<EJson::String>(Name))
+		{
+			Value = JsonObject->GetStringField(Name);
+		}
+	}
+	/**
+	 * If the underlying json object has the field, it is read into the value
+	 *
+	 * @param Name the name of the field to read
+	 * @param Value the out value to read the data into
+	 */
+	virtual void Serialize(const TCHAR* Name, FText& Value) override
+	{
+		if (JsonObject->HasTypedField<EJson::String>(Name))
+		{
+			Value = FText::FromString(JsonObject->GetStringField(Name));
+		}
+	}
+	/**
+	 * If the underlying json object has the field, it is read into the value
+	 *
+	 * @param Name the name of the field to read
+	 * @param Value the out value to read the data into
+	 */
+	virtual void Serialize(const TCHAR* Name, float& Value) override
+	{
+		if (JsonObject->HasTypedField<EJson::Number>(Name))
+		{
+			Value = JsonObject->GetNumberField(Name);
+		}
+	}
+	/**
+	* If the underlying json object has the field, it is read into the value
+	*
+	* @param Name the name of the field to read
+	* @param Value the out value to read the data into
+	*/
+	virtual void Serialize(const TCHAR* Name, double& Value) override
+	{
+		if (JsonObject->HasTypedField<EJson::Number>(Name))
+		{
+			Value = JsonObject->GetNumberField(Name);
+		}
+	}
+	/**
+	* Writes the field name and the corresponding value to the JSON data
+	*
+	* @param Name the field name to write out
+	* @param Value the value to write out
+	*/
+	virtual void Serialize(const TCHAR* Name, FDateTime& Value) override
+	{
+		if (JsonObject->HasTypedField<EJson::String>(Name))
+		{
+			FDateTime::ParseIso8601(*JsonObject->GetStringField(Name), Value);
+		}
+	}
+	/**
+	 * Serializes an array of values
+	 *
+	 * @param Name the name of the property to serialize
+	 * @param Array the array to serialize
+	 */
+	virtual void SerializeArray(FJsonSerializableArray& Array) override
+	{
+		// @todo - higher level serialization is expecting a Json Object
+		check(0 && TEXT("Not implemented"));
+	}
+	/**
+	 * Serializes an array of values with an identifier
+	 *
+	 * @param Name the name of the property to serialize
+	 * @param Array the array to serialize
+	 */
+	virtual void SerializeArray(const TCHAR* Name, FJsonSerializableArray& Array) override
+	{
+		if (JsonObject->HasTypedField<EJson::Array>(Name))
+		{
+			TArray< TSharedPtr<FJsonValue> > JsonArray = JsonObject->GetArrayField(Name);
+			// Iterate all of the keys and their values
+			for (auto ArrayIt = JsonArray.CreateConstIterator(); ArrayIt; ++ArrayIt)
+			{
+				Array.Add((*ArrayIt)->AsString());
+			}
+		}
+	}
+	/**
+	 * Serializes the keys & values for map
+	 *
+	 * @param Name the name of the property to serialize
+	 * @param Map the map to serialize
+	 */
+	virtual void SerializeMap(const TCHAR* Name, FJsonSerializableKeyValueMap& Map) override
+	{
+		if (JsonObject->HasTypedField<EJson::Object>(Name))
+		{
+			TSharedPtr<FJsonObject> JsonMap = JsonObject->GetObjectField(Name);
+			// Iterate all of the keys and their values
+			for (auto KeyValueIt = JsonMap->Values.CreateConstIterator(); KeyValueIt; ++KeyValueIt)
+			{
+				FString Value = JsonMap->GetStringField(KeyValueIt.Key());
+				Map.Add(KeyValueIt.Key(), Value);
+			}
+		}
+	}
+
+	/**
+	 * Serializes the keys & values for map
+	 *
+	 * @param Name the name of the property to serialize
+	 * @param Map the map to serialize
+	 */
+	virtual void SerializeMap(const TCHAR* Name, FJsonSerializableKeyValueMapInt& Map) override
+	{
+		if (JsonObject->HasTypedField<EJson::Object>(Name))
+		{
+			TSharedPtr<FJsonObject> JsonMap = JsonObject->GetObjectField(Name);
+			// Iterate all of the keys and their values
+			for (auto KeyValueIt = JsonMap->Values.CreateConstIterator(); KeyValueIt; ++KeyValueIt)
+			{
+				int32 Value = JsonMap->GetNumberField(KeyValueIt.Key());
+				Map.Add(KeyValueIt.Key(), Value);
+			}
+		}
+	}
+
+	virtual void WriteIdentifierPrefix(const TCHAR* Name)
+	{
+		// Should never be called on a reader
+		check(false);
+	}
+};
+
+/**
+ * Base class for a JSON serializable object
+ */
+struct FJsonSerializable
+{
+	/**
+	 * Used to allow serialization of a const ref
+	 *
+	 * @return the corresponding json string
+	 */
+	inline const FString ToJson(bool bPrettyPrint = true) const
+	{
+		// Strip away const, because we use a single method that can read/write which requires non-const semantics
+		// Otherwise, we'd have to have 2 separate macros for declaring const to json and non-const from json
+		return ((FJsonSerializable*)this)->ToJson(bPrettyPrint);
+	}
+	/**
+	 * Serializes this object to its JSON string form
+	 *
+	 * @param bPrettyPrint - If true, will use the pretty json formatter
+	 * @return the corresponding json string
+	 */
+	virtual const FString ToJson(bool bPrettyPrint=true)
+	{
+		FString JsonStr;
+		if (bPrettyPrint)
+		{
+			TSharedRef<TJsonWriter<> > JsonWriter = TJsonWriterFactory<>::Create(&JsonStr);
+			FJsonSerializerWriter<> Serializer(JsonWriter);
+			Serialize(Serializer, false);
+			JsonWriter->Close();
+		}
+		else
+		{
+			TSharedRef< TJsonWriter< TCHAR, TCondensedJsonPrintPolicy< TCHAR > > > JsonWriter = TJsonWriterFactory< TCHAR, TCondensedJsonPrintPolicy< TCHAR > >::Create( &JsonStr );
+			FJsonSerializerWriter<TCHAR, TCondensedJsonPrintPolicy< TCHAR >> Serializer(JsonWriter);
+			Serialize(Serializer, false);
+			JsonWriter->Close();
+		}
+		return JsonStr;
+	}
+	virtual void ToJson(TSharedRef<TJsonWriter<> >& JsonWriter, bool bFlatObject) const
+	{
+		FJsonSerializerWriter<> Serializer(JsonWriter);
+		((FJsonSerializable*)this)->Serialize(Serializer, bFlatObject);
+	}
+	virtual void ToJson(TSharedRef< TJsonWriter< TCHAR, TCondensedJsonPrintPolicy< TCHAR > > >& JsonWriter, bool bFlatObject) const
+	{
+		FJsonSerializerWriter<TCHAR, TCondensedJsonPrintPolicy< TCHAR >> Serializer(JsonWriter);
+		((FJsonSerializable*)this)->Serialize(Serializer, bFlatObject);
+	}
+	/**
+	 * Serializes the contents of a JSON string into this object
+	 *
+	 * @param Json the JSON data to serialize from
+	 */
+	virtual bool FromJson(const FString& Json)
+	{
+		TSharedPtr<FJsonObject> JsonObject;
+		TSharedRef<TJsonReader<> > JsonReader = TJsonReaderFactory<>::Create(Json);
+		if (FJsonSerializer::Deserialize(JsonReader,JsonObject) &&
+			JsonObject.IsValid())
+		{
+			FJsonSerializerReader Serializer(JsonObject);
+			Serialize(Serializer, false);
+			return true;
+		}
+		return false;
+	}
+	virtual bool FromJson(TSharedPtr<FJsonObject> JsonObject)
+	{
+		if (JsonObject.IsValid())
+		{
+			FJsonSerializerReader Serializer(JsonObject);
+			Serialize(Serializer, false);
+			return true;
+		}
+		return false;
+	}
+
+	/**
+	 * Abstract method that needs to be supplied using the macros
+	 *
+	 * @param Serializer the object that will perform serialization in/out of JSON
+	 * @param bFlatObject if true then no object wrapper is used
+	 */
+	virtual void Serialize(FJsonSerializerBase& Serializer, bool bFlatObject) = 0;
+};
diff --git a/Engine/Source/Runtime/Launch/Launch.Build.cs b/Engine/Source/Runtime/Launch/Launch.Build.cs
--- a/Engine/Source/Runtime/Launch/Launch.Build.cs
+++ b/Engine/Source/Runtime/Launch/Launch.Build.cs
@@ -12,7 +12,6 @@ public class Launch : ModuleRules
 		PrivateIncludePathModuleNames.AddRange(
 			new string[] {
 				"AutomationController",
-				"OnlineSubsystem",
 				"TaskGraph",
 			}
 		);
@@ -115,9 +114,9 @@ public class Launch : ModuleRules
 			PublicDependencyModuleNames.Add("SessionServices");
 			PrivateIncludePaths.Add("Developer/DerivedDataCache/Public");
 
-			// LaunchEngineLoop.cpp does a LoadModule() on OnlineSubsystem and OnlineSubsystemUtils when compiled WITH_ENGINE, so they must be marked as dependencies so that they get compiled and cleaned
-			DynamicallyLoadedModuleNames.Add("OnlineSubsystem");
-			DynamicallyLoadedModuleNames.Add("OnlineSubsystemUtils");
+            // LaunchEngineLoop.cpp will still attempt to load XMPP but not all projects require it so it will silently fail unless referenced by the project's build.cs file.
+            // DynamicallyLoadedModuleNames.Add("XMPP");
+            DynamicallyLoadedModuleNames.Add("HTTP");
 		}
 
 		if (Target.Configuration != UnrealTargetConfiguration.Shipping)
diff --git a/Engine/Source/Runtime/Launch/Private/LaunchEngineLoop.cpp b/Engine/Source/Runtime/Launch/Private/LaunchEngineLoop.cpp
--- a/Engine/Source/Runtime/Launch/Private/LaunchEngineLoop.cpp
+++ b/Engine/Source/Runtime/Launch/Private/LaunchEngineLoop.cpp
@@ -49,7 +49,6 @@
 	#include "DistanceFieldAtlas.h"
 	#include "GlobalShader.h"
 	#include "ParticleHelper.h"
-	#include "Online.h"
 	#include "PhysicsPublic.h"
 	#include "PlatformFeatures.h"
 	#include "DeviceProfiles/DeviceProfileManager.h"
@@ -61,6 +60,7 @@
 	#include "ISessionService.h"
 	#include "ISessionServicesModule.h"
 	#include "Engine/GameInstance.h"
+	#include "Net/OnlineEngineInterface.h"
 	#include "Internationalization/EnginePackageLocalizationCache.h"
 
 #if !UE_SERVER
@@ -1658,7 +1658,7 @@ int32 FEngineLoop::PreInit( const TCHAR* CmdLine )
 		GUObjectArray.CloseDisregardForGC();
 	}
 
-	if (IOnlineSubsystem::IsLoaded())
+	if (UOnlineEngineInterface::Get()->IsLoaded())
 	{
 		SetIsServerForOnlineSubsystemsDelegate(FQueryIsRunningServer::CreateStatic(&IsServerDelegateForOSS));
 	}
@@ -2162,7 +2162,7 @@ bool FEngineLoop::LoadStartupCoreModules()
 	}
 
 	// We need this for blueprint projects that have online functionality.
-	FModuleManager::Get().LoadModule(TEXT("OnlineBlueprintSupport"));
+	//FModuleManager::Get().LoadModule(TEXT("OnlineBlueprintSupport"));
 
 	if (IsRunningCommandlet())
 	{
@@ -3202,6 +3202,27 @@ bool FEngineLoop::AppInit( )
 	}
 #endif
 
+#if WITH_ENGINE
+	if (!IsRunningCommandlet())
+	{
+		// Earliest place to init the online subsystems (via plugins)
+		// Code needs GConfigFile to be valid
+		// Must be after FThreadStats::StartThread();
+		// Must be before Render/RHI subsystem D3DCreate()
+		// For platform services that need D3D hooks like Steam
+		// --
+		// Why load HTTP?
+		// Because most, if not all online subsystems will load HTTP themselves. This can cause problems though, as HTTP will be loaded *after* OSS, 
+		// and if OSS holds on to resources allocated by it, this will cause crash (modules are being unloaded in LIFO order with no dependency tracking).
+		// Loading HTTP before OSS works around this problem by making ModuleManager unload HTTP after OSS, at the cost of extra module for the few OSS (like Null) that don't use it.
+
+		// These should not be LoadModuleChecked because these modules might not exist
+		FModuleManager::Get().LoadModule(TEXT("XMPP"));
+		FModuleManager::Get().LoadModule(TEXT("HTTP"));
+		// OSS Default/Console are loaded in plugins immediately following
+	}
+#endif
+
 	// Load "pre-init" plugin modules
 	if (!IProjectManager::Get().LoadModulesForProject(ELoadingPhase::PostConfigInit) || !IPluginManager::Get().LoadModulesForEnabledPlugins(ELoadingPhase::PostConfigInit))
 	{
@@ -3305,31 +3326,6 @@ bool FEngineLoop::AppInit( )
 	FApp::InitializeSession();
 #endif
 
-#if WITH_ENGINE
-	if (!IsRunningCommandlet())
-	{
-		// Earliest place to init the online subsystems
-		// Code needs GConfigFile to be valid
-		// Must be after FThreadStats::StartThread();
-		// Must be before Render/RHI subsystem D3DCreate()
-		// For platform services that need D3D hooks like Steam
-		// --
-		// Why load HTTP?
-		// Because most, if not all online subsystems will load HTTP themselves. This can cause problems though, as HTTP will be loaded *after* OSS, 
-		// and if OSS holds on to resources allocated by it, this will cause crash (modules are being unloaded in LIFO order with no dependency tracking).
-		// Loading HTTP before OSS works around this problem by making ModuleManager unload HTTP after OSS, at the cost of extra module for the few OSS (like Null) that don't use it.
-		if (FModuleManager::Get().ModuleExists(TEXT("XMPP")))
-		{
-			FModuleManager::Get().LoadModule(TEXT("XMPP"));
-		}
-		FModuleManager::Get().LoadModule(TEXT("HTTP"));
-		FModuleManager::Get().LoadModule(TEXT("OnlineSubsystem"));
-		FModuleManager::Get().LoadModule(TEXT("OnlineSubsystemUtils"));
-		// Also load the console/platform specific OSS which might not necessarily be the default OSS instance
-		IOnlineSubsystem::GetByPlatform();
-	}
-#endif
-
 	// Checks.
 	check(sizeof(uint8) == 1);
 	check(sizeof(int8) == 1);
diff --git a/Engine/Source/Runtime/NetworkReplayStreaming/HttpNetworkReplayStreaming/HttpNetworkReplayStreaming.Build.cs b/Engine/Source/Runtime/NetworkReplayStreaming/HttpNetworkReplayStreaming/HttpNetworkReplayStreaming.Build.cs
--- a/Engine/Source/Runtime/NetworkReplayStreaming/HttpNetworkReplayStreaming/HttpNetworkReplayStreaming.Build.cs
+++ b/Engine/Source/Runtime/NetworkReplayStreaming/HttpNetworkReplayStreaming/HttpNetworkReplayStreaming.Build.cs
@@ -6,8 +6,6 @@ namespace UnrealBuildTool.Rules
     {
         public HttpNetworkReplayStreaming(TargetInfo Target)
         {
-			PrivateIncludePathModuleNames.Add( "OnlineSubsystem" );
-
 			PrivateDependencyModuleNames.AddRange(
                 new string[]
 				{
diff --git a/Engine/Source/Runtime/NetworkReplayStreaming/HttpNetworkReplayStreaming/Private/HttpNetworkReplayStreaming.cpp b/Engine/Source/Runtime/NetworkReplayStreaming/HttpNetworkReplayStreaming/Private/HttpNetworkReplayStreaming.cpp
--- a/Engine/Source/Runtime/NetworkReplayStreaming/HttpNetworkReplayStreaming/Private/HttpNetworkReplayStreaming.cpp
+++ b/Engine/Source/Runtime/NetworkReplayStreaming/HttpNetworkReplayStreaming/Private/HttpNetworkReplayStreaming.cpp
@@ -16,7 +16,7 @@ static TAutoConsoleVariable<FString> CVarMetaFilterOverride( TEXT( "httpReplay.M
 static TAutoConsoleVariable<float> CVarChunkUploadDelayInSeconds( TEXT( "httpReplay.ChunkUploadDelayInSeconds" ), 10.0f, TEXT( "" ) );
 static TAutoConsoleVariable<int> CVarMaxCacheSize( TEXT( "httpReplay.MaxCacheSize" ), 1024 * 1024 * 10, TEXT( "" ) );		// 10 MB cache by default
 
-class FNetworkReplayListItem : public FOnlineJsonSerializable
+class FNetworkReplayListItem : public FJsonSerializable
 {
 public:
 	FNetworkReplayListItem() : SizeInBytes( 0 ), DemoTimeInMs( 0 ), NumViewers( 0 ), bIsLive( false ), Changelist( 0 ), bShouldKeep( false ) {}
@@ -33,22 +33,22 @@ public:
 	int32		Changelist;
 	bool		bShouldKeep;
 
-	// FOnlineJsonSerializable
-	BEGIN_ONLINE_JSON_SERIALIZER
-		ONLINE_JSON_SERIALIZE( "AppName",		AppName );
-		ONLINE_JSON_SERIALIZE( "SessionName",	SessionName );
-		ONLINE_JSON_SERIALIZE( "FriendlyName",	FriendlyName );
-		ONLINE_JSON_SERIALIZE( "Timestamp",		Timestamp );
-		ONLINE_JSON_SERIALIZE( "SizeInBytes",	SizeInBytes );
-		ONLINE_JSON_SERIALIZE( "DemoTimeInMs",	DemoTimeInMs );
-		ONLINE_JSON_SERIALIZE( "NumViewers",	NumViewers );
-		ONLINE_JSON_SERIALIZE( "bIsLive",		bIsLive );
-		ONLINE_JSON_SERIALIZE( "Changelist",	Changelist );
-		ONLINE_JSON_SERIALIZE( "shouldKeep",	bShouldKeep );
-	END_ONLINE_JSON_SERIALIZER
+	// FJsonSerializable
+	BEGIN_JSON_SERIALIZER
+		JSON_SERIALIZE( "AppName",		AppName );
+		JSON_SERIALIZE( "SessionName",	SessionName );
+		JSON_SERIALIZE( "FriendlyName",	FriendlyName );
+		JSON_SERIALIZE( "Timestamp",		Timestamp );
+		JSON_SERIALIZE( "SizeInBytes",	SizeInBytes );
+		JSON_SERIALIZE( "DemoTimeInMs",	DemoTimeInMs );
+		JSON_SERIALIZE( "NumViewers",	NumViewers );
+		JSON_SERIALIZE( "bIsLive",		bIsLive );
+		JSON_SERIALIZE( "Changelist",	Changelist );
+		JSON_SERIALIZE( "shouldKeep",	bShouldKeep );
+	END_JSON_SERIALIZER
 };
 
-class FNetworkReplayList : public FOnlineJsonSerializable
+class FNetworkReplayList : public FJsonSerializable
 {
 public:
 	FNetworkReplayList()
@@ -57,13 +57,13 @@ public:
 
 	TArray< FNetworkReplayListItem > Replays;
 
-	// FOnlineJsonSerializable
-	BEGIN_ONLINE_JSON_SERIALIZER
-		ONLINE_JSON_SERIALIZE_ARRAY_SERIALIZABLE( "replays", Replays, FNetworkReplayListItem );
-	END_ONLINE_JSON_SERIALIZER
+	// FJsonSerializable
+	BEGIN_JSON_SERIALIZER
+		JSON_SERIALIZE_ARRAY_SERIALIZABLE( "replays", Replays, FNetworkReplayListItem );
+	END_JSON_SERIALIZER
 };
 
-class FNetworkReplayUserList : public FOnlineJsonSerializable
+class FNetworkReplayUserList : public FJsonSerializable
 {
 public:
 	FNetworkReplayUserList()
@@ -75,13 +75,13 @@ public:
 
 	TArray< FString > Users;
 
-	// FOnlineJsonSerializable
-	BEGIN_ONLINE_JSON_SERIALIZER
-		ONLINE_JSON_SERIALIZE_ARRAY( "users", Users );
-	END_ONLINE_JSON_SERIALIZER
+	// FJsonSerializable
+	BEGIN_JSON_SERIALIZER
+		JSON_SERIALIZE_ARRAY( "users", Users );
+	END_JSON_SERIALIZER
 };
 
-class FNetworkReplayStartUploadingResponse : public FOnlineJsonSerializable
+class FNetworkReplayStartUploadingResponse : public FJsonSerializable
 {
 public:
 	FNetworkReplayStartUploadingResponse()
@@ -90,13 +90,13 @@ public:
 
 	FString SessionId;
 
-	// FOnlineJsonSerializable
-	BEGIN_ONLINE_JSON_SERIALIZER
-		ONLINE_JSON_SERIALIZE("sessionId", SessionId);
-	END_ONLINE_JSON_SERIALIZER
+	// FJsonSerializable
+	BEGIN_JSON_SERIALIZER
+		JSON_SERIALIZE("sessionId", SessionId);
+	END_JSON_SERIALIZER
 };
 
-class FNetworkReplayStartDownloadingResponse : public FOnlineJsonSerializable
+class FNetworkReplayStartDownloadingResponse : public FJsonSerializable
 {
 public:
 	FNetworkReplayStartDownloadingResponse()
@@ -108,13 +108,13 @@ public:
 	int32 Time;
 	int32 NumChunks;
 
-	// FOnlineJsonSerializable
-	BEGIN_ONLINE_JSON_SERIALIZER
-		ONLINE_JSON_SERIALIZE("state", State);
-		ONLINE_JSON_SERIALIZE("numChunks", NumChunks);
-		ONLINE_JSON_SERIALIZE("time", Time);
-		ONLINE_JSON_SERIALIZE("viewerId", Viewer);
-	END_ONLINE_JSON_SERIALIZER
+	// FJsonSerializable
+	BEGIN_JSON_SERIALIZER
+		JSON_SERIALIZE("state", State);
+		JSON_SERIALIZE("numChunks", NumChunks);
+		JSON_SERIALIZE("time", Time);
+		JSON_SERIALIZE("viewerId", Viewer);
+	END_JSON_SERIALIZER
 };
 
 void FHttpStreamFArchive::Serialize( void* V, int64 Length ) 
diff --git a/Engine/Source/Runtime/NetworkReplayStreaming/HttpNetworkReplayStreaming/Public/HttpNetworkReplayStreaming.h b/Engine/Source/Runtime/NetworkReplayStreaming/HttpNetworkReplayStreaming/Public/HttpNetworkReplayStreaming.h
--- a/Engine/Source/Runtime/NetworkReplayStreaming/HttpNetworkReplayStreaming/Public/HttpNetworkReplayStreaming.h
+++ b/Engine/Source/Runtime/NetworkReplayStreaming/HttpNetworkReplayStreaming/Public/HttpNetworkReplayStreaming.h
@@ -4,7 +4,6 @@
 #include "NetworkReplayStreaming.h"
 #include "Http.h"
 #include "Runtime/Engine/Public/Tickable.h"
-#include "OnlineJsonSerializer.h"
 
 /**
  * Archive used to buffer stream over http
diff --git a/Engine/Source/Runtime/NetworkReplayStreaming/InMemoryNetworkReplayStreaming/Public/InMemoryNetworkReplayStreaming.h b/Engine/Source/Runtime/NetworkReplayStreaming/InMemoryNetworkReplayStreaming/Public/InMemoryNetworkReplayStreaming.h
--- a/Engine/Source/Runtime/NetworkReplayStreaming/InMemoryNetworkReplayStreaming/Public/InMemoryNetworkReplayStreaming.h
+++ b/Engine/Source/Runtime/NetworkReplayStreaming/InMemoryNetworkReplayStreaming/Public/InMemoryNetworkReplayStreaming.h
@@ -6,7 +6,7 @@
 #include "Core.h"
 #include "ModuleManager.h"
 #include "UniquePtr.h"
-#include "OnlineJsonSerializer.h"
+#include "JsonSerializer.h"
 #include "Tickable.h"
 
 class FInMemoryNetworkReplayStreamingFactory;
diff --git a/Engine/Source/Runtime/NetworkReplayStreaming/NetworkReplayStreaming/NetworkReplayStreaming.Build.cs b/Engine/Source/Runtime/NetworkReplayStreaming/NetworkReplayStreaming/NetworkReplayStreaming.Build.cs
--- a/Engine/Source/Runtime/NetworkReplayStreaming/NetworkReplayStreaming/NetworkReplayStreaming.Build.cs
+++ b/Engine/Source/Runtime/NetworkReplayStreaming/NetworkReplayStreaming/NetworkReplayStreaming.Build.cs
@@ -8,8 +8,6 @@ namespace UnrealBuildTool.Rules
 		{
 			PrivateIncludePaths.Add( "Runtime/NetworkReplayStreaming/NetworkReplayStreaming/Private" );
 
-            PublicIncludePathModuleNames.Add("OnlineSubsystem");
-
 			PublicDependencyModuleNames.AddRange(
 				new string[]
 				{
diff --git a/Engine/Source/Runtime/NetworkReplayStreaming/NetworkReplayStreaming/Public/NetworkReplayStreaming.h b/Engine/Source/Runtime/NetworkReplayStreaming/NetworkReplayStreaming/Public/NetworkReplayStreaming.h
--- a/Engine/Source/Runtime/NetworkReplayStreaming/NetworkReplayStreaming/Public/NetworkReplayStreaming.h
+++ b/Engine/Source/Runtime/NetworkReplayStreaming/NetworkReplayStreaming/Public/NetworkReplayStreaming.h
@@ -5,10 +5,10 @@
 // Dependencies.
 #include "ModuleManager.h"
 #include "Core.h"
-#include "OnlineJsonSerializer.h"
+#include "Json.h"
 #include "NetworkVersion.h"
 
-class FReplayEventListItem : public FOnlineJsonSerializable
+class FReplayEventListItem : public FJsonSerializable
 {
 public:
 	FReplayEventListItem() {}
@@ -20,17 +20,17 @@ public:
 	uint32		Time1;
 	uint32		Time2;
 
-	// FOnlineJsonSerializable
-	BEGIN_ONLINE_JSON_SERIALIZER
-	ONLINE_JSON_SERIALIZE("id", ID);
-	ONLINE_JSON_SERIALIZE("group", Group);
-	ONLINE_JSON_SERIALIZE("meta", Metadata);
-	ONLINE_JSON_SERIALIZE("time1", Time1);
-	ONLINE_JSON_SERIALIZE("time2", Time2);
-	END_ONLINE_JSON_SERIALIZER
+	// FJsonSerializable
+	BEGIN_JSON_SERIALIZER
+		JSON_SERIALIZE("id", ID);
+		JSON_SERIALIZE("group", Group);
+		JSON_SERIALIZE("meta", Metadata);
+		JSON_SERIALIZE("time1", Time1);
+		JSON_SERIALIZE("time2", Time2);
+	END_JSON_SERIALIZER
 };
 
-class FReplayEventList : public FOnlineJsonSerializable
+class FReplayEventList : public FJsonSerializable
 {
 public:
 	FReplayEventList()
@@ -39,10 +39,10 @@ public:
 
 	TArray< FReplayEventListItem > ReplayEvents;
 
-	// FOnlineJsonSerializable
-	BEGIN_ONLINE_JSON_SERIALIZER
-		ONLINE_JSON_SERIALIZE_ARRAY_SERIALIZABLE("events", ReplayEvents, FReplayEventListItem);
-	END_ONLINE_JSON_SERIALIZER
+	// FJsonSerializable
+	BEGIN_JSON_SERIALIZER
+		JSON_SERIALIZE_ARRAY_SERIALIZABLE("events", ReplayEvents, FReplayEventListItem);
+	END_JSON_SERIALIZER
 };
 
 /** Struct to store information about a stream, returned from search results. */
diff --git a/Engine/Source/Runtime/NetworkReplayStreaming/NullNetworkReplayStreaming/NullNetworkReplayStreaming.build.cs b/Engine/Source/Runtime/NetworkReplayStreaming/NullNetworkReplayStreaming/NullNetworkReplayStreaming.build.cs
--- a/Engine/Source/Runtime/NetworkReplayStreaming/NullNetworkReplayStreaming/NullNetworkReplayStreaming.build.cs
+++ b/Engine/Source/Runtime/NetworkReplayStreaming/NullNetworkReplayStreaming/NullNetworkReplayStreaming.build.cs
@@ -8,8 +8,6 @@ namespace UnrealBuildTool.Rules
 		{
 			PrivateIncludePaths.Add( "Runtime/NetworkReplayStreaming/NullNetworkReplayStreaming/Private" );
 
-            PrivateIncludePathModuleNames.Add("OnlineSubsystem");
-
 			PrivateDependencyModuleNames.AddRange(
 				new string[]
 				{
diff --git a/Engine/Source/Runtime/NetworkReplayStreaming/NullNetworkReplayStreaming/Private/NullNetworkReplayStreaming.cpp b/Engine/Source/Runtime/NetworkReplayStreaming/NullNetworkReplayStreaming/Private/NullNetworkReplayStreaming.cpp
--- a/Engine/Source/Runtime/NetworkReplayStreaming/NullNetworkReplayStreaming/Private/NullNetworkReplayStreaming.cpp
+++ b/Engine/Source/Runtime/NetworkReplayStreaming/NullNetworkReplayStreaming/Private/NullNetworkReplayStreaming.cpp
@@ -7,7 +7,7 @@
 DEFINE_LOG_CATEGORY_STATIC( LogNullReplay, Log, All );
 
 /* Class to hold stream event information */
-class FNullCheckpointListItem : public FOnlineJsonSerializable
+class FNullCheckpointListItem : public FJsonSerializable
 {
 public:
 	FNullCheckpointListItem() : Time1(0), Time2(0) {}
@@ -17,13 +17,13 @@ public:
 	uint32		Time1;
 	uint32		Time2;
 
-	// FOnlineJsonSerializable
-	BEGIN_ONLINE_JSON_SERIALIZER
-		ONLINE_JSON_SERIALIZE( "group",			Group );
-		ONLINE_JSON_SERIALIZE( "meta",			Metadata );
-		ONLINE_JSON_SERIALIZE( "time1",			Time1 );
-		ONLINE_JSON_SERIALIZE( "time2",			Time2 );
-	END_ONLINE_JSON_SERIALIZER
+	// FJsonSerializable
+	BEGIN_JSON_SERIALIZER
+		JSON_SERIALIZE( "group",			Group );
+		JSON_SERIALIZE( "meta",			Metadata );
+		JSON_SERIALIZE( "time1",			Time1 );
+		JSON_SERIALIZE( "time2",			Time2 );
+	END_JSON_SERIALIZER
 };
 
 /**
diff --git a/Engine/Source/Runtime/NetworkReplayStreaming/NullNetworkReplayStreaming/Public/NullNetworkReplayStreaming.h b/Engine/Source/Runtime/NetworkReplayStreaming/NullNetworkReplayStreaming/Public/NullNetworkReplayStreaming.h
--- a/Engine/Source/Runtime/NetworkReplayStreaming/NullNetworkReplayStreaming/Public/NullNetworkReplayStreaming.h
+++ b/Engine/Source/Runtime/NetworkReplayStreaming/NullNetworkReplayStreaming/Public/NullNetworkReplayStreaming.h
@@ -7,11 +7,11 @@
 #include "Engine.h"
 #include "ModuleManager.h"
 #include "UniquePtr.h"
-#include "OnlineJsonSerializer.h"
+#include "Json.h"
 #include "Tickable.h"
 
 /* Class to hold metadata about an entire replay */
-class FNullReplayInfo : public FOnlineJsonSerializable
+class FNullReplayInfo : public FJsonSerializable
 {
 public:
 	FNullReplayInfo() : LengthInMS(0), NetworkVersion(0), Changelist(0), DemoFileLastOffset(0), bIsValid(false) {}
@@ -23,14 +23,14 @@ public:
 	int32		DemoFileLastOffset;
 	bool		bIsValid;
 
-	// FOnlineJsonSerializable
-	BEGIN_ONLINE_JSON_SERIALIZER
-		ONLINE_JSON_SERIALIZE( "LengthInMS",		LengthInMS );
-		ONLINE_JSON_SERIALIZE( "NetworkVersion",	NetworkVersion );
-		ONLINE_JSON_SERIALIZE( "Changelist",		Changelist );
-		ONLINE_JSON_SERIALIZE( "FriendlyName",		FriendlyName );
-		ONLINE_JSON_SERIALIZE( "DemoFileLastOffset",	DemoFileLastOffset );
-	END_ONLINE_JSON_SERIALIZER
+	// FJsonSerializable
+	BEGIN_JSON_SERIALIZER
+		JSON_SERIALIZE( "LengthInMS",		LengthInMS );
+		JSON_SERIALIZE( "NetworkVersion",	NetworkVersion );
+		JSON_SERIALIZE( "Changelist",		Changelist );
+		JSON_SERIALIZE( "FriendlyName",		FriendlyName );
+		JSON_SERIALIZE( "DemoFileLastOffset",	DemoFileLastOffset );
+	END_JSON_SERIALIZER
 };
 
 /** Default streamer that goes straight to the HD */
diff --git a/Engine/Source/Runtime/Slate/Private/Framework/Text/SlateTextLayout.cpp b/Engine/Source/Runtime/Slate/Private/Framework/Text/SlateTextLayout.cpp
--- a/Engine/Source/Runtime/Slate/Private/Framework/Text/SlateTextLayout.cpp
+++ b/Engine/Source/Runtime/Slate/Private/Framework/Text/SlateTextLayout.cpp
@@ -14,8 +14,8 @@ TSharedRef< FSlateTextLayout > FSlateTextLayout::Create(FTextBlockStyle InDefaul
 }
 
 FSlateTextLayout::FSlateTextLayout(FTextBlockStyle InDefaultTextStyle)
-	: Children()
-	, DefaultTextStyle(MoveTemp(InDefaultTextStyle))
+	: DefaultTextStyle(MoveTemp(InDefaultTextStyle))
+	, Children()
 	, bIsPassword(false)
 	, LocalizedFallbackFontRevision(INDEX_NONE)
 {
diff --git a/Engine/Source/Runtime/Slate/Private/Widgets/Input/SEditableText.cpp b/Engine/Source/Runtime/Slate/Private/Widgets/Input/SEditableText.cpp
--- a/Engine/Source/Runtime/Slate/Private/Widgets/Input/SEditableText.cpp
+++ b/Engine/Source/Runtime/Slate/Private/Widgets/Input/SEditableText.cpp
@@ -54,7 +54,7 @@ void SEditableText::Construct( const FArguments& InArgs )
 	// We use a separate marshaller for the hint text, as that should never be displayed as a password
 	TSharedRef<FPlainTextLayoutMarshaller> HintTextMarshaller = FPlainTextLayoutMarshaller::Create();
 
-	EditableTextLayout = MakeUnique<FSlateEditableTextLayout>(*this, InArgs._Text, TextStyle, InArgs._TextShapingMethod, InArgs._TextFlowDirection, PlainTextMarshaller.ToSharedRef(), HintTextMarshaller);
+	EditableTextLayout = MakeUnique<FSlateEditableTextLayout>(*this, InArgs._Text, TextStyle, InArgs._TextShapingMethod, InArgs._TextFlowDirection, FCreateSlateTextLayout(), PlainTextMarshaller.ToSharedRef(), HintTextMarshaller);
 	EditableTextLayout->SetHintText(InArgs._HintText);
 	EditableTextLayout->SetCursorBrush(InArgs._CaretImage.IsSet() ? InArgs._CaretImage : &InArgs._Style->CaretImage);
 	EditableTextLayout->SetCompositionBrush(InArgs._BackgroundImageComposing.IsSet() ? InArgs._BackgroundImageComposing : &InArgs._Style->BackgroundImageComposing);
diff --git a/Engine/Source/Runtime/Slate/Private/Widgets/Input/SMultiLineEditableTextBox.cpp b/Engine/Source/Runtime/Slate/Private/Widgets/Input/SMultiLineEditableTextBox.cpp
--- a/Engine/Source/Runtime/Slate/Private/Widgets/Input/SMultiLineEditableTextBox.cpp
+++ b/Engine/Source/Runtime/Slate/Private/Widgets/Input/SMultiLineEditableTextBox.cpp
@@ -108,6 +108,7 @@ void SMultiLineEditableTextBox::Construct( const FArguments& InArgs )
 					.OnTextCommitted( InArgs._OnTextCommitted )
 					.OnCursorMoved( InArgs._OnCursorMoved )
 					.ContextMenuExtender( InArgs._ContextMenuExtender )
+					.CreateSlateTextLayout( InArgs._CreateSlateTextLayout )
 					.Justification(InArgs._Justification)
 					.RevertTextOnEscape(InArgs._RevertTextOnEscape)
 					.SelectAllTextWhenFocused(InArgs._SelectAllTextWhenFocused)
diff --git a/Engine/Source/Runtime/Slate/Private/Widgets/Text/SMultiLineEditableText.cpp b/Engine/Source/Runtime/Slate/Private/Widgets/Text/SMultiLineEditableText.cpp
--- a/Engine/Source/Runtime/Slate/Private/Widgets/Text/SMultiLineEditableText.cpp
+++ b/Engine/Source/Runtime/Slate/Private/Widgets/Text/SMultiLineEditableText.cpp
@@ -67,7 +67,7 @@ void SMultiLineEditableText::Construct( const FArguments& InArgs )
 		Marshaller = FPlainTextLayoutMarshaller::Create();
 	}
 
-	EditableTextLayout = MakeUnique<FSlateEditableTextLayout>(*this, InArgs._Text, TextStyle, InArgs._TextShapingMethod, InArgs._TextFlowDirection, Marshaller.ToSharedRef(), Marshaller.ToSharedRef());
+	EditableTextLayout = MakeUnique<FSlateEditableTextLayout>(*this, InArgs._Text, TextStyle, InArgs._TextShapingMethod, InArgs._TextFlowDirection, InArgs._CreateSlateTextLayout, Marshaller.ToSharedRef(), Marshaller.ToSharedRef());
 	EditableTextLayout->SetHintText(InArgs._HintText);
 	EditableTextLayout->SetTextWrapping(InArgs._WrapTextAt, InArgs._AutoWrapText, InArgs._WrappingPolicy);
 	EditableTextLayout->SetMargin(InArgs._Margin);
diff --git a/Engine/Source/Runtime/Slate/Private/Widgets/Text/SRichTextBlock.cpp b/Engine/Source/Runtime/Slate/Private/Widgets/Text/SRichTextBlock.cpp
--- a/Engine/Source/Runtime/Slate/Private/Widgets/Text/SRichTextBlock.cpp
+++ b/Engine/Source/Runtime/Slate/Private/Widgets/Text/SRichTextBlock.cpp
@@ -52,7 +52,7 @@ void SRichTextBlock::Construct( const FArguments& InArgs )
 			Marshaller->AppendInlineDecorator(Decorator);
 		}
 
-		TextLayoutCache = MakeUnique<FTextBlockLayout>(TextStyle, InArgs._TextShapingMethod, InArgs._TextFlowDirection, Marshaller.ToSharedRef(), nullptr);
+		TextLayoutCache = MakeUnique<FTextBlockLayout>(TextStyle, InArgs._TextShapingMethod, InArgs._TextFlowDirection, InArgs._CreateSlateTextLayout, Marshaller.ToSharedRef(), nullptr);
 		TextLayoutCache->SetDebugSourceInfo(TAttribute<FString>::Create(TAttribute<FString>::FGetter::CreateLambda([this]{ return FReflectionMetaData::GetWidgetDebugInfo(this); })));
 	}
 }
diff --git a/Engine/Source/Runtime/Slate/Private/Widgets/Text/STextBlock.cpp b/Engine/Source/Runtime/Slate/Private/Widgets/Text/STextBlock.cpp
--- a/Engine/Source/Runtime/Slate/Private/Widgets/Text/STextBlock.cpp
+++ b/Engine/Source/Runtime/Slate/Private/Widgets/Text/STextBlock.cpp
@@ -46,7 +46,7 @@ void STextBlock::Construct( const FArguments& InArgs )
 	BoundText = InArgs._Text;
 
 	// We use a dummy style here (as it may not be safe to call the delegates used to compute the style), but the correct style is set by ComputeDesiredSize
-	TextLayoutCache = MakeUnique<FTextBlockLayout>(FTextBlockStyle::GetDefault(), InArgs._TextShapingMethod, InArgs._TextFlowDirection, FPlainTextLayoutMarshaller::Create(), InArgs._LineBreakPolicy);
+	TextLayoutCache = MakeUnique<FTextBlockLayout>(FTextBlockStyle::GetDefault(), InArgs._TextShapingMethod, InArgs._TextFlowDirection, FCreateSlateTextLayout(), FPlainTextLayoutMarshaller::Create(), InArgs._LineBreakPolicy);
 	TextLayoutCache->SetDebugSourceInfo(TAttribute<FString>::Create(TAttribute<FString>::FGetter::CreateLambda([this]{ return FReflectionMetaData::GetWidgetDebugInfo(this); })));
 }
 
diff --git a/Engine/Source/Runtime/Slate/Private/Widgets/Text/SlateEditableTextLayout.cpp b/Engine/Source/Runtime/Slate/Private/Widgets/Text/SlateEditableTextLayout.cpp
--- a/Engine/Source/Runtime/Slate/Private/Widgets/Text/SlateEditableTextLayout.cpp
+++ b/Engine/Source/Runtime/Slate/Private/Widgets/Text/SlateEditableTextLayout.cpp
@@ -58,13 +58,19 @@ bool IsCharAllowed(const TCHAR InChar)
 
 }
 
-FSlateEditableTextLayout::FSlateEditableTextLayout(ISlateEditableTextWidget& InOwnerWidget, const TAttribute<FText>& InInitialText, FTextBlockStyle InTextStyle, const TOptional<ETextShapingMethod> InTextShapingMethod, const TOptional<ETextFlowDirection> InTextFlowDirection, TSharedRef<ITextLayoutMarshaller> InTextMarshaller, TSharedRef<ITextLayoutMarshaller> InHintTextMarshaller)
+FSlateEditableTextLayout::FSlateEditableTextLayout(ISlateEditableTextWidget& InOwnerWidget, const TAttribute<FText>& InInitialText, FTextBlockStyle InTextStyle, const TOptional<ETextShapingMethod> InTextShapingMethod, const TOptional<ETextFlowDirection> InTextFlowDirection, const FCreateSlateTextLayout& InCreateSlateTextLayout, TSharedRef<ITextLayoutMarshaller> InTextMarshaller, TSharedRef<ITextLayoutMarshaller> InHintTextMarshaller)
 {
+	CreateSlateTextLayout = InCreateSlateTextLayout;
+	if (!CreateSlateTextLayout.IsBound())
+	{
+		CreateSlateTextLayout.BindStatic(&FSlateTextLayout::Create);
+	}
+
 	OwnerWidget = &InOwnerWidget;
 	Marshaller = InTextMarshaller;
 	HintMarshaller = InHintTextMarshaller;
 	TextStyle = InTextStyle;
-	TextLayout = FSlateTextLayout::Create(TextStyle);
+	TextLayout = CreateSlateTextLayout.Execute(TextStyle);
 
 	WrapTextAt = 0.0f;
 	AutoWrapText = false;
@@ -250,7 +256,7 @@ void FSlateEditableTextLayout::SetHintText(const TAttribute<FText>& InHintText)
 	if (HintText.IsBound() || !HintText.Get(FText::GetEmpty()).IsEmpty())
 	{
 		HintTextStyle = TextStyle;
-		HintTextLayout = MakeUnique<FTextBlockLayout>(HintTextStyle, TextLayout->GetTextShapingMethod(), TextLayout->GetTextFlowDirection(), HintMarshaller.ToSharedRef(), nullptr);
+		HintTextLayout = MakeUnique<FTextBlockLayout>(HintTextStyle, TextLayout->GetTextShapingMethod(), TextLayout->GetTextFlowDirection(), CreateSlateTextLayout, HintMarshaller.ToSharedRef(), nullptr);
 		HintTextLayout->SetDebugSourceInfo(DebugSourceInfo);
 	}
 	else
diff --git a/Engine/Source/Runtime/Slate/Private/Widgets/Text/TextBlockLayout.cpp b/Engine/Source/Runtime/Slate/Private/Widgets/Text/TextBlockLayout.cpp
--- a/Engine/Source/Runtime/Slate/Private/Widgets/Text/TextBlockLayout.cpp
+++ b/Engine/Source/Runtime/Slate/Private/Widgets/Text/TextBlockLayout.cpp
@@ -9,8 +9,8 @@
 
 SLATE_DECLARE_CYCLE_COUNTER(GSlateTextBlockLayoutComputeDesiredSize, "FTextBlockLayout ComputeDesiredSize");
 
-FTextBlockLayout::FTextBlockLayout(FTextBlockStyle InDefaultTextStyle, const TOptional<ETextShapingMethod> InTextShapingMethod, const TOptional<ETextFlowDirection> InTextFlowDirection, TSharedRef<ITextLayoutMarshaller> InMarshaller, TSharedPtr<IBreakIterator> InLineBreakPolicy)
-	: TextLayout(FSlateTextLayout::Create(MoveTemp(InDefaultTextStyle)))
+FTextBlockLayout::FTextBlockLayout(FTextBlockStyle InDefaultTextStyle, const TOptional<ETextShapingMethod> InTextShapingMethod, const TOptional<ETextFlowDirection> InTextFlowDirection, const FCreateSlateTextLayout& InCreateSlateTextLayout, TSharedRef<ITextLayoutMarshaller> InMarshaller, TSharedPtr<IBreakIterator> InLineBreakPolicy)
+	: TextLayout((InCreateSlateTextLayout.IsBound()) ? InCreateSlateTextLayout.Execute(MoveTemp(InDefaultTextStyle)) : FSlateTextLayout::Create(MoveTemp(InDefaultTextStyle)))
 	, Marshaller(MoveTemp(InMarshaller))
 	, TextHighlighter(FSlateTextHighlightRunRenderer::Create())
 	, CachedSize(ForceInitToZero)
diff --git a/Engine/Source/Runtime/Slate/Public/Framework/Text/SlateTextLayout.h b/Engine/Source/Runtime/Slate/Public/Framework/Text/SlateTextLayout.h
--- a/Engine/Source/Runtime/Slate/Public/Framework/Text/SlateTextLayout.h
+++ b/Engine/Source/Runtime/Slate/Public/Framework/Text/SlateTextLayout.h
@@ -9,9 +9,9 @@ public:
 
 	FChildren* GetChildren();
 
-	void ArrangeChildren( const FGeometry& AllottedGeometry, FArrangedChildren& ArrangedChildren ) const;
+	virtual void ArrangeChildren( const FGeometry& AllottedGeometry, FArrangedChildren& ArrangedChildren ) const;
 
-	int32 OnPaint( const FPaintArgs& Args, const FGeometry& AllottedGeometry, const FSlateRect& MyClippingRect, FSlateWindowElementList& OutDrawElements, int32 LayerId, const FWidgetStyle& InWidgetStyle, bool bParentEnabled ) const;
+	virtual int32 OnPaint( const FPaintArgs& Args, const FGeometry& AllottedGeometry, const FSlateRect& MyClippingRect, FSlateWindowElementList& OutDrawElements, int32 LayerId, const FWidgetStyle& InWidgetStyle, bool bParentEnabled ) const;
 
 	virtual void EndLayout() override;
 
@@ -26,24 +26,23 @@ protected:
 
 	FSlateTextLayout(FTextBlockStyle InDefaultTextStyle);
 
-	int32 OnPaintHighlights(const FPaintArgs& Args, const FTextLayout::FLineView& LineView, const TArray<FLineViewHighlight>& Highlights, const FTextBlockStyle& DefaultTextStyle, const FGeometry& AllottedGeometry, const FSlateRect& ClippingRect, FSlateWindowElementList& OutDrawElements, const int32 LayerId, const FWidgetStyle& InWidgetStyle, bool bParentEnabled) const;
+	virtual int32 OnPaintHighlights(const FPaintArgs& Args, const FTextLayout::FLineView& LineView, const TArray<FLineViewHighlight>& Highlights, const FTextBlockStyle& DefaultTextStyle, const FGeometry& AllottedGeometry, const FSlateRect& ClippingRect, FSlateWindowElementList& OutDrawElements, const int32 LayerId, const FWidgetStyle& InWidgetStyle, bool bParentEnabled) const;
 
-	void AggregateChildren();
+	virtual void AggregateChildren();
 
 	virtual TSharedRef<IRun> CreateDefaultTextRun(const TSharedRef<FString>& NewText, const FTextRange& NewRange) const override;
 
+protected:
+	/** Default style used by the TextLayout */
+	FTextBlockStyle DefaultTextStyle;
+
 private:
 
 	TSlotlessChildren< SWidget > Children;
 
-	/** Default style used by the TextLayout */
-	FTextBlockStyle DefaultTextStyle;
-
 	/** This this layout displaying a password? */
 	TAttribute<bool> bIsPassword;
 
 	/** The localized fallback font revision the last time the text layout was updated. Used to force a flush if the font changes. */
 	int32 LocalizedFallbackFontRevision;
-
-	friend class FSlateTextLayoutFactory;
 };
diff --git a/Engine/Source/Runtime/Slate/Public/Framework/Text/SlateTextLayoutFactory.h b/Engine/Source/Runtime/Slate/Public/Framework/Text/SlateTextLayoutFactory.h
new file mode 100644
--- /dev/null
+++ b/Engine/Source/Runtime/Slate/Public/Framework/Text/SlateTextLayoutFactory.h
@@ -0,0 +1,7 @@
+// Copyright 1998-2016 Epic Games, Inc. All Rights Reserved.
+
+#pragma once
+
+class FSlateTextLayout;
+
+DECLARE_DELEGATE_RetVal_OneParam(TSharedRef<FSlateTextLayout>, FCreateSlateTextLayout, FTextBlockStyle);
diff --git a/Engine/Source/Runtime/Slate/Public/Widgets/Input/SMultiLineEditableTextBox.h b/Engine/Source/Runtime/Slate/Public/Widgets/Input/SMultiLineEditableTextBox.h
--- a/Engine/Source/Runtime/Slate/Public/Widgets/Input/SMultiLineEditableTextBox.h
+++ b/Engine/Source/Runtime/Slate/Public/Widgets/Input/SMultiLineEditableTextBox.h
@@ -140,6 +140,9 @@ public:
 		/** Menu extender for the right-click context menu */
 		SLATE_EVENT( FMenuExtensionDelegate, ContextMenuExtender )
 
+		/** Delegate used to create text layouts for this widget. If none is provided then FSlateTextLayout will be used. */
+		SLATE_EVENT( FCreateSlateTextLayout, CreateSlateTextLayout )
+
 		/** Whether text wraps onto a new line when it's length exceeds this width; if this value is zero or negative, no wrapping occurs. */
 		SLATE_ATTRIBUTE( float, WrapTextAt )
 
diff --git a/Engine/Source/Runtime/Slate/Public/Widgets/Text/SMultiLineEditableText.h b/Engine/Source/Runtime/Slate/Public/Widgets/Text/SMultiLineEditableText.h
--- a/Engine/Source/Runtime/Slate/Public/Widgets/Text/SMultiLineEditableText.h
+++ b/Engine/Source/Runtime/Slate/Public/Widgets/Text/SMultiLineEditableText.h
@@ -5,6 +5,7 @@
 #if WITH_FANCY_TEXT
 
 #include "ISlateEditableTextWidget.h"
+#include "SlateTextLayoutFactory.h"
 #include "SScrollBar.h"
 #include "UniquePtr.h"
 
@@ -127,6 +128,9 @@ public:
 		/** Menu extender for the right-click context menu */
 		SLATE_EVENT(FMenuExtensionDelegate, ContextMenuExtender)
 
+		/** Delegate used to create text layouts for this widget. If none is provided then FSlateTextLayout will be used. */
+		SLATE_EVENT(FCreateSlateTextLayout, CreateSlateTextLayout)
+
 		/** The optional modifier key necessary to create a newline when typing into the editor. */
 		SLATE_ARGUMENT(EModifierKey::Type, ModiferKeyForNewLine)
 
diff --git a/Engine/Source/Runtime/Slate/Public/Widgets/Text/SRichTextBlock.h b/Engine/Source/Runtime/Slate/Public/Widgets/Text/SRichTextBlock.h
--- a/Engine/Source/Runtime/Slate/Public/Widgets/Text/SRichTextBlock.h
+++ b/Engine/Source/Runtime/Slate/Public/Widgets/Text/SRichTextBlock.h
@@ -2,6 +2,7 @@
 #pragma once
 
 #include "UniquePtr.h"
+#include "SlateTextLayoutFactory.h"
 
 #if WITH_FANCY_TEXT
 
@@ -51,6 +52,9 @@ public:
 		/** The marshaller used to get/set the raw text to/from the text layout. */
 		SLATE_ARGUMENT(TSharedPtr<class FRichTextLayoutMarshaller>, Marshaller)
 
+		/** Delegate used to create text layouts for this widget. If none is provided then FSlateTextLayout will be used. */
+		SLATE_EVENT(FCreateSlateTextLayout, CreateSlateTextLayout)
+
 		/** The style set used for looking up styles used by decorators*/
 		SLATE_ARGUMENT( const ISlateStyle*, DecoratorStyleSet )
 
diff --git a/Engine/Source/Runtime/Slate/Public/Widgets/Text/SlateEditableTextLayout.h b/Engine/Source/Runtime/Slate/Public/Widgets/Text/SlateEditableTextLayout.h
--- a/Engine/Source/Runtime/Slate/Public/Widgets/Text/SlateEditableTextLayout.h
+++ b/Engine/Source/Runtime/Slate/Public/Widgets/Text/SlateEditableTextLayout.h
@@ -3,6 +3,7 @@
 #pragma once
 
 #include "SlateEditableTextTypes.h"
+#include "SlateTextLayoutFactory.h"
 #include "IVirtualKeyboardEntry.h"
 #include "ITextInputMethodSystem.h"
 #include "ISlateEditableTextWidget.h"
@@ -18,7 +19,7 @@ class ISlateRunRenderer;
 class SLATE_API FSlateEditableTextLayout
 {
 public:
-	FSlateEditableTextLayout(ISlateEditableTextWidget& InOwnerWidget, const TAttribute<FText>& InInitialText, FTextBlockStyle InTextStyle, const TOptional<ETextShapingMethod> InTextShapingMethod, const TOptional<ETextFlowDirection> InTextFlowDirection, TSharedRef<ITextLayoutMarshaller> InTextMarshaller, TSharedRef<ITextLayoutMarshaller> InHintTextMarshaller);
+	FSlateEditableTextLayout(ISlateEditableTextWidget& InOwnerWidget, const TAttribute<FText>& InInitialText, FTextBlockStyle InTextStyle, const TOptional<ETextShapingMethod> InTextShapingMethod, const TOptional<ETextFlowDirection> InTextFlowDirection, const FCreateSlateTextLayout& InCreateSlateTextLayout, TSharedRef<ITextLayoutMarshaller> InTextMarshaller, TSharedRef<ITextLayoutMarshaller> InHintTextMarshaller);
 	~FSlateEditableTextLayout();
 
 	void SetText(const TAttribute<FText>& InText);
@@ -438,6 +439,9 @@ private:
 	/** The marshaller used to get/set the HintText text to/from the text layout. */
 	TSharedPtr<ITextLayoutMarshaller> HintMarshaller;
 
+	/** Delegate used to create internal text layouts. */
+	FCreateSlateTextLayout CreateSlateTextLayout;
+
 	/** In control of the layout and wrapping of the BoundText */
 	TSharedPtr<FSlateTextLayout> TextLayout;
 
diff --git a/Engine/Source/Runtime/Slate/Public/Widgets/Text/TextBlockLayout.h b/Engine/Source/Runtime/Slate/Public/Widgets/Text/TextBlockLayout.h
--- a/Engine/Source/Runtime/Slate/Public/Widgets/Text/TextBlockLayout.h
+++ b/Engine/Source/Runtime/Slate/Public/Widgets/Text/TextBlockLayout.h
@@ -2,6 +2,8 @@
 
 #pragma once
 
+#include "SlateTextLayoutFactory.h"
+
 class FSlateTextLayout;
 class ITextLayoutMarshaller;
 class ISlateRunRenderer;
@@ -46,7 +48,7 @@ public:
 	/**
 	 * Constructor
 	 */
-	FTextBlockLayout(FTextBlockStyle InDefaultTextStyle, const TOptional<ETextShapingMethod> InTextShapingMethod, const TOptional<ETextFlowDirection> InTextFlowDirection, TSharedRef<ITextLayoutMarshaller> InMarshaller, TSharedPtr<IBreakIterator> InLineBreakPolicy);
+	FTextBlockLayout(FTextBlockStyle InDefaultTextStyle, const TOptional<ETextShapingMethod> InTextShapingMethod, const TOptional<ETextFlowDirection> InTextFlowDirection, const FCreateSlateTextLayout& InCreateSlateTextLayout, TSharedRef<ITextLayoutMarshaller> InMarshaller, TSharedPtr<IBreakIterator> InLineBreakPolicy);
 
 	/**
 	 * Get the computed desired size for this layout, updating the internal cache as required
diff --git a/Engine/Source/Runtime/UE4Game/UE4Game.Build.cs b/Engine/Source/Runtime/UE4Game/UE4Game.Build.cs
--- a/Engine/Source/Runtime/UE4Game/UE4Game.Build.cs
+++ b/Engine/Source/Runtime/UE4Game/UE4Game.Build.cs
@@ -8,32 +8,23 @@ public class UE4Game : ModuleRules
 	{
 		PrivateDependencyModuleNames.Add("Core");
 	
-		DynamicallyLoadedModuleNames.Add("OnlineSubsystemNull");
+        //DynamicallyLoadedModuleNames.Add("OnlineSubsystemNull");
 
-		if (Target.Platform == UnrealTargetPlatform.Win32 ||
-			Target.Platform == UnrealTargetPlatform.Win64 ||
-			Target.Platform == UnrealTargetPlatform.Mac)
-		{
-			if (UEBuildConfiguration.bCompileSteamOSS == true)
-			{
-				DynamicallyLoadedModuleNames.Add("OnlineSubsystemSteam");
-			}
-		}
 		if (Target.Platform == UnrealTargetPlatform.IOS || Target.Platform == UnrealTargetPlatform.TVOS)
 		{
-			PrivateDependencyModuleNames.AddRange(new string[] { "Core", "CoreUObject", "Engine", "OnlineSubsystem", "OnlineSubsystemUtils" });
-			DynamicallyLoadedModuleNames.Add("OnlineSubsystemIOS");
+			PrivateDependencyModuleNames.AddRange(new string[] { "Core", "CoreUObject", "Engine" });//, "OnlineSubsystem", "OnlineSubsystemUtils" });
+            //DynamicallyLoadedModuleNames.Add("OnlineSubsystemIOS");
 			if (Target.Platform == UnrealTargetPlatform.IOS)
 			{
-				DynamicallyLoadedModuleNames.Add("OnlineSubsystemFacebook");
+                //DynamicallyLoadedModuleNames.Add("OnlineSubsystemFacebook");
 				DynamicallyLoadedModuleNames.Add("IOSAdvertising");
 			}
 		}
 		else if (Target.Platform == UnrealTargetPlatform.Android)
 		{
-			PrivateDependencyModuleNames.AddRange(new string[] { "Core", "CoreUObject", "Engine", "OnlineSubsystem", "OnlineSubsystemUtils" });
+			PrivateDependencyModuleNames.AddRange(new string[] { "Core", "CoreUObject", "Engine" });//, "OnlineSubsystem", "OnlineSubsystemUtils" });
 			DynamicallyLoadedModuleNames.Add("AndroidAdvertising");
-			DynamicallyLoadedModuleNames.Add("OnlineSubsystemGooglePlay");
+            //DynamicallyLoadedModuleNames.Add("OnlineSubsystemGooglePlay");
 		}
 	}
 }
diff --git a/Engine/Source/UE4Game.Target.cs b/Engine/Source/UE4Game.Target.cs
--- a/Engine/Source/UE4Game.Target.cs
+++ b/Engine/Source/UE4Game.Target.cs
@@ -26,10 +26,10 @@ public class UE4GameTarget : TargetRules
 	{
 		OutExtraModuleNames.Add("UE4Game");
 		// this is important - for some reason achievements etc intertwined with the onlinesubsystem and they saved without using a fake OSS. :/
-		if (Target.Platform == UnrealTargetPlatform.HTML5)
-		{
-			OutExtraModuleNames.Add("OnlineSubsystemNull");
-		}
+		//if (Target.Platform == UnrealTargetPlatform.HTML5)
+		//{
+		//	OutExtraModuleNames.Add("OnlineSubsystemNull");
+		//}
 	}
 
 	public override void SetupGlobalEnvironment(
