- adding an FJsonSerializable struct that exposes the underlying FJsonObject to the app code for cases where the model isn't rigid (or where deserialization is deferred)
- Going to convert LinksServiceMcp's link metadata to this since that's more accurate to the contract (currently we only support a map of string->string, should be string->arbitrary_json)

diff --git a/Engine/Source/Runtime/Json/Public/Serialization/JsonSerializerMacros.h b/Engine/Source/Runtime/Json/Public/Serialization/JsonSerializerMacros.h
--- a/Engine/Source/Runtime/Json/Public/Serialization/JsonSerializerMacros.h
+++ b/Engine/Source/Runtime/Json/Public/Serialization/JsonSerializerMacros.h
@@ -849,3 +849,40 @@ struct FJsonSerializable
 	 */
 	virtual void Serialize(FJsonSerializerBase& Serializer, bool bFlatObject) = 0;
 };
+
+/**
+ * Useful if you just want access to the underlying FJsonObject (for cases where the schema is loose or an outer system will do further de/serialization)
+ */
+struct FJsonDataBag
+	: public FJsonSerializable
+{
+	virtual void Serialize(FJsonSerializerBase& Serializer, bool bFlatObject) override
+	{
+		if (Serializer.IsLoading())
+		{
+			// just grab a reference to the underlying JSON object
+			JsonObject = Serializer.GetObject();
+		}
+		else
+		{
+			if (JsonObject.IsValid())
+			{
+				// if we have an object, serialize to string and write raw
+				FString JsonStr;
+				auto Writer = TJsonWriterFactory<TCHAR, TCondensedJsonPrintPolicy<TCHAR>>::Create(&JsonStr);
+				FJsonSerializer::Serialize(JsonObject.ToSharedRef(), Writer);
+				// too bad there's no JsonObject serialization method on FJsonSerializerBase directly :-/
+				Serializer.WriteRawJSONValue(*JsonStr);
+			}
+			else if (!bFlatObject)
+			{
+				// no object, just write empty
+				Serializer.StartObject();
+				Serializer.EndObject();
+			}
+		}
+	}
+
+public:
+	TSharedPtr<FJsonObject> JsonObject;
+};
