Adding helpful GetField and SetField methods to FJsonDataBag

diff --git a/Engine/Source/Runtime/Json/Public/Serialization/JsonSerializerMacros.h b/Engine/Source/Runtime/Json/Public/Serialization/JsonSerializerMacros.h
--- a/Engine/Source/Runtime/Json/Public/Serialization/JsonSerializerMacros.h
+++ b/Engine/Source/Runtime/Json/Public/Serialization/JsonSerializerMacros.h
@@ -883,6 +883,30 @@ struct FJsonDataBag
 		}
 	}
 
+	TSharedPtr<const FJsonValue> GetField(const FString& Key) const
+	{
+		if (JsonObject.IsValid())
+		{
+			return JsonObject->TryGetField(Key);
+		}
+		return TSharedPtr<const FJsonValue>();
+	}
+
+	template<typename JSON_TYPE, typename Arg>
+	void SetField(const FString& Key, Arg&& Value)
+	{
+		SetFieldJson(Key, MakeShared<JSON_TYPE>(MoveTempIfPossible(Value)));
+	}
+
+	void SetFieldJson(const FString& Key, const TSharedPtr<FJsonValue>& Value)
+	{
+		if (!JsonObject.IsValid())
+		{
+			JsonObject = MakeShared<FJsonObject>();
+		}
+		JsonObject->SetField(Key, Value);
+	}
+
 public:
 	TSharedPtr<FJsonObject> JsonObject;
 };
