Added support for bFlatObject in FJsonDataBag

diff --git a/Engine/Source/Runtime/Json/Public/Serialization/JsonSerializerMacros.h b/Engine/Source/Runtime/Json/Public/Serialization/JsonSerializerMacros.h
--- a/Engine/Source/Runtime/Json/Public/Serialization/JsonSerializerMacros.h
+++ b/Engine/Source/Runtime/Json/Public/Serialization/JsonSerializerMacros.h
@@ -865,19 +865,57 @@ struct FJsonDataBag
 		}
 		else
 		{
+			if (!bFlatObject)
+			{
+				Serializer.StartObject();
+			}
+
 			if (JsonObject.IsValid())
 			{
-				// if we have an object, serialize to string and write raw
-				FString JsonStr;
-				auto Writer = TJsonWriterFactory<TCHAR, TCondensedJsonPrintPolicy<TCHAR>>::Create(&JsonStr);
-				FJsonSerializer::Serialize(JsonObject.ToSharedRef(), Writer);
-				// too bad there's no JsonObject serialization method on FJsonSerializerBase directly :-/
-				Serializer.WriteRawJSONValue(*JsonStr);
+				for (const auto& It : JsonObject->Values)
+				{
+					TSharedPtr<FJsonValue> JsonValue = It.Value;
+					if (JsonValue.IsValid())
+					{
+						switch (JsonValue->Type)
+						{
+							case EJson::Boolean:
+							{
+								auto Value = JsonValue->AsBool();
+								Serializer.Serialize(*It.Key, Value);
+								break;
+							}
+							case EJson::Number:
+							{
+								auto Value = JsonValue->AsNumber();
+								Serializer.Serialize(*It.Key, Value);
+								break;
+							}
+							case EJson::String:
+							{
+								auto Value = JsonValue->AsString();
+								Serializer.Serialize(*It.Key, Value);
+								break;
+							}
+							case EJson::Array:
+							case EJson::Object:
+							{
+								// if we have an object, serialize to string and write raw
+								FString JsonStr;
+								auto Writer = TJsonWriterFactory<TCHAR, TCondensedJsonPrintPolicy<TCHAR>>::Create(&JsonStr);
+								FJsonSerializer::Serialize(JsonValue->AsObject().ToSharedRef(), Writer);
+								// too bad there's no JsonObject serialization method on FJsonSerializerBase directly :-/
+								Serializer.WriteIdentifierPrefix(*It.Key);
+								Serializer.WriteRawJSONValue(*JsonStr);
+								break;
+							}
+						}
+					}
+				}
 			}
-			else if (!bFlatObject)
+
+			if (!bFlatObject)
 			{
-				// no object, just write empty
-				Serializer.StartObject();
 				Serializer.EndObject();
 			}
 		}
