Fix TJsonWriter.WriteStringValue() to honor low-ascii control characters without a two-character code by escaping them with \u00XX.
#tests local unit test that verifies the correct output.

diff --git a/Engine/Source/Runtime/Json/Public/Serialization/JsonWriter.h b/Engine/Source/Runtime/Json/Public/Serialization/JsonWriter.h
--- a/Engine/Source/Runtime/Json/Public/Serialization/JsonWriter.h
+++ b/Engine/Source/Runtime/Json/Public/Serialization/JsonWriter.h
@@ -352,7 +352,16 @@ protected:
 			case TCHAR('\f'): OutString += TEXT("\\f"); break;
 			case TCHAR('\r'): OutString += TEXT("\\r"); break;
 			case TCHAR('\"'): OutString += TEXT("\\\""); break;
-			default: OutString += *Char;
+			default: 
+				// Must escape control characters
+				if (*Char >= TCHAR(32))
+				{
+					OutString += *Char;
+				}
+				else
+				{
+					OutString += FString::Printf(TEXT("\\u%04x"), *Char);
+				}
 			}
 		}
 		OutString += TEXT("\"");
