Adding typed get helpers for FJsonDataBag

diff --git a/Engine/Source/Runtime/Json/Public/Serialization/JsonSerializerMacros.h b/Engine/Source/Runtime/Json/Public/Serialization/JsonSerializerMacros.h
--- a/Engine/Source/Runtime/Json/Public/Serialization/JsonSerializerMacros.h
+++ b/Engine/Source/Runtime/Json/Public/Serialization/JsonSerializerMacros.h
@@ -921,6 +921,24 @@ struct FJsonDataBag
 		}
 	}
 
+	double GetDouble(const FString& Key) const
+	{
+		const auto Json = GetField(Key);
+		return Json.IsValid() ? Json->AsNumber() : 0.0;
+	}
+
+	FString GetString(const FString& Key) const
+	{
+		const auto Json = GetField(Key);
+		return Json.IsValid() ? Json->AsString() : FString();
+	}
+
+	bool GetBool(const FString& Key) const
+	{
+		const auto Json = GetField(Key);
+		return Json.IsValid() ? Json->AsBool() : false;
+	}
+
 	TSharedPtr<const FJsonValue> GetField(const FString& Key) const
 	{
 		if (JsonObject.IsValid())
