Fix build Editor Linux

diff --git a/Engine/Source/Runtime/Core/Private/Misc/AutomationTest.cpp b/Engine/Source/Runtime/Core/Private/Misc/AutomationTest.cpp
--- a/Engine/Source/Runtime/Core/Private/Misc/AutomationTest.cpp
+++ b/Engine/Source/Runtime/Core/Private/Misc/AutomationTest.cpp
@@ -12,7 +12,7 @@
 #include "Modules/ModuleManager.h"
 #include "Misc/OutputDeviceRedirector.h"
 #include "Internationalization/Regex.h"
-
+#include <inttypes.h>
 
 DEFINE_LOG_CATEGORY_STATIC(LogAutomationTest, Warning, All);
 
@@ -1078,6 +1078,15 @@ void FAutomationTestBase::TestEqual(const TCHAR* What, const int32 Actual, const
 	}
 }
 
+
+void FAutomationTestBase::TestEqual(const TCHAR* What, const int64 Actual, const int64 Expected)
+{
+	if (Actual != Expected)
+	{
+		AddError(FString::Printf(TEXT("Expected '%s' to be %" PRId64 ", but it was %" PRId64 "."), What, Expected, Actual), 1);
+	}
+}
+
 void FAutomationTestBase::TestEqual(const TCHAR* What, const float Actual, const float Expected, float Tolerance)
 {
 	if (!FMath::IsNearlyEqual(Actual, Expected, Tolerance))
@@ -1164,4 +1173,4 @@ bool FAutomationTestBase::IsExpectedError(const FString& Error)
 	}
 
 	return false;
-}
\ No newline at end of file
+}
diff --git a/Engine/Source/Runtime/Core/Public/Misc/AutomationTest.h b/Engine/Source/Runtime/Core/Public/Misc/AutomationTest.h
--- a/Engine/Source/Runtime/Core/Public/Misc/AutomationTest.h
+++ b/Engine/Source/Runtime/Core/Public/Misc/AutomationTest.h
@@ -1231,6 +1231,7 @@ public:
 public:
 
 	void TestEqual(const TCHAR* What, int32 Actual, int32 Expected);
+	void TestEqual(const TCHAR* What, int64 Actual, int64 Expected);
 	void TestEqual(const TCHAR* What, float Actual, float Expected, float Tolerance = KINDA_SMALL_NUMBER);
 	void TestEqual(const TCHAR* What, double Actual, double Expected, double Tolerance = KINDA_SMALL_NUMBER);
 	void TestEqual(const TCHAR* What, FVector Actual, FVector Expected, float Tolerance = KINDA_SMALL_NUMBER);
diff --git a/Engine/Source/Runtime/Json/Private/Tests/JsonTests.cpp b/Engine/Source/Runtime/Json/Private/Tests/JsonTests.cpp
--- a/Engine/Source/Runtime/Json/Private/Tests/JsonTests.cpp
+++ b/Engine/Source/Runtime/Json/Private/Tests/JsonTests.cpp
@@ -713,28 +713,28 @@ bool FJsonAutomationTest::RunTest(const FString& Parameters)
 			int64 IntVal;
 			bool bOk = JsonNumberToInt64(9007199254740991.0, IntVal);
 			TestTrue(TEXT("TryGetNumber-Int64 Big Float64 succeeds"), bOk);
-			TestEqual(TEXT("TryGetNumber-Int64 Big Float64"), IntVal, 9007199254740991L);
+			TestEqual(TEXT("TryGetNumber-Int64 Big Float64"), IntVal, 9007199254740991LL);
 		}
 
 		{
 			int64 IntVal;
 			bool bOk = JsonNumberToInt64(-9007199254740991.0, IntVal);
 			TestTrue(TEXT("TryGetNumber-Int64 Small Float64 succeeds"), bOk);
-			TestEqual(TEXT("TryGetNumber-Int64 Small Float64"), IntVal, -9007199254740991L);
+			TestEqual(TEXT("TryGetNumber-Int64 Small Float64"), IntVal, -9007199254740991LL);
 		}
 
 		{
 			int64 IntVal;
 			bool bOk = JsonNumberToInt64(0.4999999999999997, IntVal);
 			TestTrue(TEXT("TryGetNumber-Int64 Lesser than near half succeeds"), bOk);
-			TestEqual(TEXT("TryGetNumber-Int64 Lesser than near half rounds to zero"), IntVal, 0);
+			TestEqual(TEXT("TryGetNumber-Int64 Lesser than near half rounds to zero"), IntVal, 0LL);
 		}
 
 		{
 			int64 IntVal;
 			bool bOk = JsonNumberToInt64(-0.4999999999999997, IntVal);
 			TestTrue(TEXT("TryGetNumber-Int64 Greater than near negative half succeeds"), bOk);
-			TestEqual(TEXT("TryGetNumber-Int64 Greater than near negative half rounds to zero"), IntVal, 0);
+			TestEqual(TEXT("TryGetNumber-Int64 Greater than near negative half rounds to zero"), IntVal, 0LL);
 		}
 
 		{
