build fix: android protection pass

diff --git a/Engine/Source/Runtime/Json/Public/Serialization/JsonReader.h b/Engine/Source/Runtime/Json/Public/Serialization/JsonReader.h
--- a/Engine/Source/Runtime/Json/Public/Serialization/JsonReader.h
+++ b/Engine/Source/Runtime/Json/Public/Serialization/JsonReader.h
@@ -8,22 +8,30 @@
 
 class Error;
 
-static EJsonNotation TokenToNotationTable[] = 
-{ 
-	EJsonNotation::Error,			/*EJsonToken::None*/
-	EJsonNotation::Error,			/*EJsonToken::Comma*/
-	EJsonNotation::ObjectStart,		/*EJsonToken::CurlyOpen*/
-	EJsonNotation::ObjectEnd,		/*EJsonToken::CurlyClose*/
-	EJsonNotation::ArrayStart,		/*EJsonToken::SquareOpen*/
-	EJsonNotation::ArrayEnd,		/*EJsonToken::SquareClose*/
-	EJsonNotation::Error,			/*EJsonToken::Colon*/
-	EJsonNotation::String,			/*EJsonToken::String*/
-	EJsonNotation::Number,			/*EJsonToken::Number*/
-	EJsonNotation::Boolean,			/*EJsonToken::True*/
-	EJsonNotation::Boolean,			/*EJsonToken::False*/
-	EJsonNotation::Null,			/*EJsonToken::Null*/
+#define JSON_NOTATIONMAP_DEF \
+static EJsonNotation TokenToNotationTable[] =  \
+{ \
+	EJsonNotation::Error,			/*EJsonToken::None*/ \
+	EJsonNotation::Error,			/*EJsonToken::Comma*/ \
+	EJsonNotation::ObjectStart,		/*EJsonToken::CurlyOpen*/ \
+	EJsonNotation::ObjectEnd,		/*EJsonToken::CurlyClose*/ \
+	EJsonNotation::ArrayStart,		/*EJsonToken::SquareOpen*/ \
+	EJsonNotation::ArrayEnd,		/*EJsonToken::SquareClose*/ \
+	EJsonNotation::Error,			/*EJsonToken::Colon*/ \
+	EJsonNotation::String,			/*EJsonToken::String*/ \
+	EJsonNotation::Number,			/*EJsonToken::Number*/ \
+	EJsonNotation::Boolean,			/*EJsonToken::True*/ \
+	EJsonNotation::Boolean,			/*EJsonToken::False*/ \
+	EJsonNotation::Null,			/*EJsonToken::Null*/ \
 };
 
+#ifndef WITH_JSON_INLINED_NOTATIONMAP
+#define WITH_JSON_INLINED_NOTATIONMAP 0
+#endif // WITH_JSON_INLINED_NOTATIONMAP
+
+#if !WITH_JSON_INLINED_NOTATIONMAP
+JSON_NOTATIONMAP_DEF;
+#endif // WITH_JSON_INLINED_NOTATIONMAP
 
 template <class CharType = TCHAR>
 class TJsonReader
@@ -104,6 +112,10 @@ public:
 		}
 		while (ReadWasSuccess && (CurrentToken == EJsonToken::None));
 
+#if WITH_JSON_INLINED_NOTATIONMAP
+		JSON_NOTATIONMAP_DEF;
+#endif // WITH_JSON_INLINED_NOTATIONMAP
+
 		Notation = TokenToNotationTable[(int32)CurrentToken];
 		FinishedReadingRootObject = ParseState.Num() == 0;
 
