Fix parsing error when FJsonReader reads a unicode character froim a stream with a narrower char type

diff --git a/Engine/Source/Runtime/Json/Public/Serialization/JsonReader.h b/Engine/Source/Runtime/Json/Public/Serialization/JsonReader.h
--- a/Engine/Source/Runtime/Json/Public/Serialization/JsonReader.h
+++ b/Engine/Source/Runtime/Json/Public/Serialization/JsonReader.h
@@ -563,7 +563,7 @@ private:
 							HexNum += HexDigit * FMath::Pow(16, Radix);
 						}
 
-						String += (CharType)HexNum;
+						String += (FString::ElementType)HexNum;
 					}
 					break;
 
