[Json] Add support for keeping numbers stored as text in the DOM, such that it can be correctly converted to the requested type when the value is queries, rather than be cast to double at read time, potentially losing data.

diff --git a/Engine/Source/Runtime/CoreUObject/Private/Serialization/Formatters/JsonArchiveInputFormatter.cpp b/Engine/Source/Runtime/CoreUObject/Private/Serialization/Formatters/JsonArchiveInputFormatter.cpp
--- a/Engine/Source/Runtime/CoreUObject/Private/Serialization/Formatters/JsonArchiveInputFormatter.cpp
+++ b/Engine/Source/Runtime/CoreUObject/Private/Serialization/Formatters/JsonArchiveInputFormatter.cpp
@@ -22,7 +22,7 @@ FJsonArchiveInputFormatter::FJsonArchiveInputFormatter(FArchive& InInner, TFunct
 
 	TSharedPtr< FJsonObject > RootObject;
 	TSharedRef< TJsonReader<char> > Reader = TJsonReaderFactory<char>::Create(&InInner);
-	ensure(FJsonSerializer::Deserialize(Reader, RootObject));
+	ensure(FJsonSerializer::Deserialize(Reader, RootObject, FJsonSerializer::EFlags::StoreNumbersAsStrings));
 	ValueStack.Add(MakeShared<FJsonValueObject>(RootObject));
 
 	ValueStack.Reserve(64);
@@ -284,52 +284,52 @@ bool FJsonArchiveInputFormatter::TryEnterAttribute(FArchiveFieldName AttributeNa
 
 void FJsonArchiveInputFormatter::Serialize(uint8& Value)
 {
-	Value = (uint8)ValueStack.Top()->AsNumber();
+	ValueStack.Top()->TryGetNumber(Value);
 }
 
 void FJsonArchiveInputFormatter::Serialize(uint16& Value)
 {
-	Value = (uint16)ValueStack.Top()->AsNumber();
+	ValueStack.Top()->TryGetNumber(Value);
 }
 
 void FJsonArchiveInputFormatter::Serialize(uint32& Value)
 {
-	Value = (uint32)ValueStack.Top()->AsNumber();
+	ValueStack.Top()->TryGetNumber(Value);
 }
 
 void FJsonArchiveInputFormatter::Serialize(uint64& Value)
 {
-	Value = (uint64)ValueStack.Top()->AsNumber();
+	ValueStack.Top()->TryGetNumber(Value);
 }
 
 void FJsonArchiveInputFormatter::Serialize(int8& Value)
 {
-	Value = (int8)ValueStack.Top()->AsNumber();
+	ValueStack.Top()->TryGetNumber(Value);
 }
 
 void FJsonArchiveInputFormatter::Serialize(int16& Value)
 {
-	Value = (int16)ValueStack.Top()->AsNumber();
+	ValueStack.Top()->TryGetNumber(Value);;
 }
 
 void FJsonArchiveInputFormatter::Serialize(int32& Value)
 {
-	Value = (int32)ValueStack.Top()->AsNumber();
+	ValueStack.Top()->TryGetNumber(Value);
 }
 
 void FJsonArchiveInputFormatter::Serialize(int64& Value)
 {
-	Value = (int64)ValueStack.Top()->AsNumber();
+	ValueStack.Top()->TryGetNumber(Value);
 }
 
 void FJsonArchiveInputFormatter::Serialize(float& Value)
 {
-	Value = (float)ValueStack.Top()->AsNumber();
+	ValueStack.Top()->TryGetNumber(Value);
 }
 
 void FJsonArchiveInputFormatter::Serialize(double& Value)
 {
-	Value = (double)ValueStack.Top()->AsNumber();
+	ValueStack.Top()->TryGetNumber(Value);
 }
 
 void FJsonArchiveInputFormatter::Serialize(bool& Value)
diff --git a/Engine/Source/Runtime/Json/Private/Dom/JsonValue.cpp b/Engine/Source/Runtime/Json/Private/Dom/JsonValue.cpp
--- a/Engine/Source/Runtime/Json/Private/Dom/JsonValue.cpp
+++ b/Engine/Source/Runtime/Json/Private/Dom/JsonValue.cpp
@@ -72,51 +72,73 @@ const TSharedPtr<FJsonObject>& FJsonValue::AsObject() const
 	return *Object;
 }
 
-
-bool FJsonValue::TryGetNumber( int32& OutNumber ) const
+template <typename T>
+bool TryConvertNumber(const FJsonValue& InValue, T& OutNumber)
 {
 	double Double;
 
-	if (TryGetNumber(Double) && (Double >= INT_MIN) && (Double <= INT_MAX))
+	if (InValue.TryGetNumber(Double) && (Double >= TNumericLimits<T>::Min()) && (Double <= TNumericLimits<T>::Max()))
 	{
-		OutNumber = static_cast<int32>(FMath::RoundHalfFromZero(Double));
-		
+		OutNumber = static_cast<T>(FMath::RoundHalfFromZero(Double));
+
 		return true;
 	}
 
 	return false;
 }
 
-
-bool FJsonValue::TryGetNumber( uint32& OutNumber ) const
+bool FJsonValue::TryGetNumber(float& OutNumber) const
 {
 	double Double;
 
-	if (TryGetNumber(Double) && (Double >= 0.0) && (Double <= UINT_MAX))
+	if (TryGetNumber(Double))
 	{
-		OutNumber = static_cast<uint32>(FMath::RoundHalfFromZero(Double));
-
+		OutNumber = static_cast<float>(Double);
 		return true;
 	}
 
 	return false;
 }
 
+bool FJsonValue::TryGetNumber(uint8& OutNumber) const
+{
+	return TryConvertNumber(*this, OutNumber);
+}
 
-bool FJsonValue::TryGetNumber( int64& OutNumber ) const
+bool FJsonValue::TryGetNumber(uint16& OutNumber) const
 {
-	double Double;
+	return TryConvertNumber(*this, OutNumber);
+}
 
-	if (TryGetNumber(Double) && (Double >= INT64_MIN) && (Double <= INT64_MAX))
-	{
-		OutNumber = static_cast<int64>(FMath::RoundHalfFromZero(Double));
-		
-		return true;
-	}
+bool FJsonValue::TryGetNumber(uint32& OutNumber) const
+{
+	return TryConvertNumber(*this, OutNumber);
+}
 
-	return false;
+bool FJsonValue::TryGetNumber(uint64& OutNumber) const
+{
+	return TryConvertNumber(*this, OutNumber);
 }
 
+bool FJsonValue::TryGetNumber(int8& OutNumber) const
+{
+	return TryConvertNumber(*this, OutNumber);
+}
+
+bool FJsonValue::TryGetNumber(int16& OutNumber) const
+{
+	return TryConvertNumber(*this, OutNumber);
+}
+
+bool FJsonValue::TryGetNumber(int32& OutNumber) const
+{
+	return TryConvertNumber(*this, OutNumber);
+}
+
+bool FJsonValue::TryGetNumber(int64& OutNumber) const
+{
+	return TryConvertNumber(*this, OutNumber);
+}
 
 //static 
 bool FJsonValue::CompareEqual( const FJsonValue& Lhs, const FJsonValue& Rhs )
diff --git a/Engine/Source/Runtime/Json/Public/Dom/JsonValue.h b/Engine/Source/Runtime/Json/Public/Dom/JsonValue.h
--- a/Engine/Source/Runtime/Json/Public/Dom/JsonValue.h
+++ b/Engine/Source/Runtime/Json/Public/Dom/JsonValue.h
@@ -31,16 +31,34 @@ public:
 	virtual const TSharedPtr<FJsonObject>& AsObject() const;
 
 	/** Tries to convert this value to a number, returning false if not possible */
-	virtual bool TryGetNumber(double& OutDouble) const { return false; }
+	virtual bool TryGetNumber(double& OutNumber) const { return false; }
 
-	/** Tries to convert this value to an integer, returning false if not possible */
-	bool TryGetNumber(int32& OutNumber) const;
+	/** Tries to convert this value to a number, returning false if not possible */
+	virtual bool TryGetNumber(float& OutNumber) const;
+
+	/** Tries to convert this value to a number, returning false if not possible */
+	virtual bool TryGetNumber(int8& OutNumber) const;
 
-	/** Tries to convert this value to an integer, returning false if not possible */
-	bool TryGetNumber(uint32& OutNumber) const;
+	/** Tries to convert this value to a number, returning false if not possible */
+	virtual bool TryGetNumber(int16& OutNumber) const;
 
-	/** Tries to convert this value to an integer, returning false if not possible */
-	bool TryGetNumber(int64& OutNumber) const;
+	/** Tries to convert this value to a number, returning false if not possible */
+	virtual bool TryGetNumber(int32& OutNumber) const;
+
+	/** Tries to convert this value to a number, returning false if not possible */
+	virtual bool TryGetNumber(int64& OutNumber) const;
+
+	/** Tries to convert this value to a number, returning false if not possible */
+	virtual bool TryGetNumber(uint8& OutNumber) const;
+
+	/** Tries to convert this value to a number, returning false if not possible */
+	virtual bool TryGetNumber(uint16& OutNumber) const;
+
+	/** Tries to convert this value to a number, returning false if not possible */
+	virtual bool TryGetNumber(uint32& OutNumber) const;
+
+	/** Tries to convert this value to a number, returning false if not possible */
+	virtual bool TryGetNumber(uint64& OutNumber) const;
 
 	/** Tries to convert this value to a string, returning false if not possible */
 	virtual bool TryGetString(FString& OutString) const { return false; }
@@ -99,7 +117,11 @@ public:
 	FJsonValueString(const FString& InString) : Value(InString) {Type = EJson::String;}
 
 	virtual bool TryGetString(FString& OutString) const override	{ OutString = Value; return true; }
-	virtual bool TryGetNumber(double& OutDouble) const override		{ if(Value.IsNumeric()){ OutDouble = FCString::Atod(*Value); return true; } else { return false; } }
+	virtual bool TryGetNumber(double& OutDouble) const override		{ if (Value.IsNumeric()) { OutDouble = FCString::Atod(*Value); return true; } else { return false; } }
+	virtual bool TryGetNumber(int32& OutValue) const override		{ LexFromString(OutValue, *Value); return true; }
+	virtual bool TryGetNumber(uint32& OutValue) const override		{ LexFromString(OutValue, *Value); return true; }
+	virtual bool TryGetNumber(int64& OutValue) const override		{ LexFromString(OutValue, *Value); return true; }
+	virtual bool TryGetNumber(uint64& OutValue) const override		{ LexFromString(OutValue, *Value); return true; }
 	virtual bool TryGetBool(bool& OutBool) const override			{ OutBool = Value.ToBool(); return true; }
 
 protected:
@@ -119,12 +141,39 @@ public:
 	virtual bool TryGetString(FString& OutString) const override	{ OutString = FString::SanitizeFloat(Value, 0); return true; }
 	
 protected:
+
 	double Value;
 
 	virtual FString GetType() const override {return TEXT("Number");}
 };
 
 
+/** A Json Number Value, stored internally as a string so as not to lose precision */
+class JSON_API FJsonValueNumberString : public FJsonValue
+{
+public:
+	FJsonValueNumberString(const FString& InString) : Value(InString) { Type = EJson::Number; }
+
+	virtual bool TryGetString(FString& OutString) const override { OutString = Value; return true; }
+	virtual bool TryGetNumber(double& OutDouble) const override { return LexTryParseString(OutDouble, *Value); }
+	virtual bool TryGetNumber(float &OutDouble) const override { return LexTryParseString(OutDouble, *Value); }
+	virtual bool TryGetNumber(int8& OutValue) const override { return LexTryParseString(OutValue, *Value); }
+	virtual bool TryGetNumber(int16& OutValue) const override { return LexTryParseString(OutValue, *Value); }
+	virtual bool TryGetNumber(int32& OutValue) const override { return LexTryParseString(OutValue, *Value); }
+	virtual bool TryGetNumber(int64& OutValue) const override { return LexTryParseString(OutValue, *Value); }
+	virtual bool TryGetNumber(uint8& OutValue) const override { return LexTryParseString(OutValue, *Value); }
+	virtual bool TryGetNumber(uint16& OutValue) const override { return LexTryParseString(OutValue, *Value); }
+	virtual bool TryGetNumber(uint32& OutValue) const override { return LexTryParseString(OutValue, *Value); }
+	virtual bool TryGetNumber(uint64& OutValue) const override { return LexTryParseString(OutValue, *Value); }
+	virtual bool TryGetBool(bool& OutBool) const override { OutBool = Value.ToBool(); return true; }
+
+protected:
+	FString Value;
+
+	virtual FString GetType() const override { return TEXT("NumberString"); }
+};
+
+
 /** A Json Boolean Value. */
 class JSON_API FJsonValueBoolean : public FJsonValue
 {
diff --git a/Engine/Source/Runtime/Json/Public/Serialization/JsonReader.h b/Engine/Source/Runtime/Json/Public/Serialization/JsonReader.h
--- a/Engine/Source/Runtime/Json/Public/Serialization/JsonReader.h
+++ b/Engine/Source/Runtime/Json/Public/Serialization/JsonReader.h
@@ -150,6 +150,12 @@ public:
 		check(CurrentToken == EJsonToken::Number);
 		return NumberValue;
 	}
+
+	FORCEINLINE const FString& GetValueAsNumberString() const
+	{
+		check(CurrentToken == EJsonToken::Number);
+		return StringValue;
+	}
 	
 	FORCEINLINE bool GetValueAsBoolean() const 
 	{ 
@@ -697,6 +703,7 @@ private:
 		// ensure the number has followed valid Json format
 		if (!Error && ((State == 2) || (State == 3) || (State == 6) || (State == 8)))
 		{
+			StringValue = String;
 			NumberValue = FCString::Atod(*String);
 			return true;
 		}
diff --git a/Engine/Source/Runtime/Json/Public/Serialization/JsonSerializer.h b/Engine/Source/Runtime/Json/Public/Serialization/JsonSerializer.h
--- a/Engine/Source/Runtime/Json/Public/Serialization/JsonSerializer.h
+++ b/Engine/Source/Runtime/Json/Public/Serialization/JsonSerializer.h
@@ -15,17 +15,23 @@ class FJsonSerializer
 {
 public:
 
+	enum class EFlags
+	{
+		None = 0,
+		StoreNumbersAsStrings = 1,
+	};
+
 	template <class CharType>
-	static bool Deserialize(const TSharedRef<TJsonReader<CharType>>& Reader, TArray<TSharedPtr<FJsonValue>>& OutArray)
+	static bool Deserialize(const TSharedRef<TJsonReader<CharType>>& Reader, TArray<TSharedPtr<FJsonValue>>& OutArray, EFlags InOptions = EFlags::None)
 	{
-		return Deserialize(*Reader, OutArray);
+		return Deserialize(*Reader, OutArray, InOptions);
 	}
 
 	template <class CharType>
-	static bool Deserialize(TJsonReader<CharType>& Reader, TArray<TSharedPtr<FJsonValue>>& OutArray )
+	static bool Deserialize(TJsonReader<CharType>& Reader, TArray<TSharedPtr<FJsonValue>>& OutArray, EFlags InOptions = EFlags::None)
 	{
 		StackState State;
-		if (!Deserialize(Reader, /*OUT*/State))
+		if (!Deserialize(Reader, /*OUT*/State, InOptions))
 		{
 			return false;
 		}
@@ -42,15 +48,15 @@ public:
 	}
 
 	template <class CharType>
-	static bool Deserialize(const TSharedRef<TJsonReader<CharType>>& Reader, TSharedPtr<FJsonObject>& OutObject )
+	static bool Deserialize(const TSharedRef<TJsonReader<CharType>>& Reader, TSharedPtr<FJsonObject>& OutObject, EFlags InOptions = EFlags::None)
 	{
-		return Deserialize(*Reader, OutObject);
+		return Deserialize(*Reader, OutObject, InOptions);
 	}
 	template <class CharType>
-	static bool Deserialize(TJsonReader<CharType>& Reader, TSharedPtr<FJsonObject>& OutObject)
+	static bool Deserialize(TJsonReader<CharType>& Reader, TSharedPtr<FJsonObject>& OutObject, EFlags InOptions = EFlags::None)
 	{
 		StackState State;
-		if (!Deserialize(Reader, /*OUT*/State))
+		if (!Deserialize(Reader, /*OUT*/State, InOptions))
 		{
 			return false;
 		}
@@ -66,16 +72,16 @@ public:
 	}
 
 	template <class CharType>
-	static bool Deserialize(const TSharedRef<TJsonReader<CharType>>& Reader, TSharedPtr<FJsonValue>& OutValue)
+	static bool Deserialize(const TSharedRef<TJsonReader<CharType>>& Reader, TSharedPtr<FJsonValue>& OutValue, EFlags InOptions = EFlags::None)
 	{
-		return Deserialize(*Reader, OutValue);
+		return Deserialize(*Reader, OutValue, InOptions);
 	}
 
 	template <class CharType>
-	static bool Deserialize(TJsonReader<CharType>& Reader, TSharedPtr<FJsonValue>& OutValue)
+	static bool Deserialize(TJsonReader<CharType>& Reader, TSharedPtr<FJsonValue>& OutValue, EFlags InOptions = EFlags::None)
 	{
 		StackState State;
-		if (!Deserialize(Reader, /*OUT*/State))
+		if (!Deserialize(Reader, /*OUT*/State, InOptions))
 		{
 			return false;
 		}
@@ -183,7 +189,7 @@ private:
 private:
 
 	template <class CharType>
-	static bool Deserialize(TJsonReader<CharType>& Reader, StackState& OutStackState)
+	static bool Deserialize(TJsonReader<CharType>& Reader, StackState& OutStackState, EFlags InOptions)
 	{
 		TArray<TSharedRef<StackState>> ScopeStack; 
 		TSharedPtr<StackState> CurrentState;
@@ -256,7 +262,14 @@ private:
 				break;
 
 			case EJsonNotation::Number:
-				NewValue = MakeShared<FJsonValueNumber>(Reader.GetValueAsNumber());
+				if (EnumHasAnyFlags(InOptions, EFlags::StoreNumbersAsStrings))
+				{
+					NewValue = MakeShared<FJsonValueNumberString>(Reader.GetValueAsNumberString());
+				}
+				else
+				{
+					NewValue = MakeShared<FJsonValueNumber>(Reader.GetValueAsNumber());
+				}
 				break;
 
 			case EJsonNotation::Null:
