- MMS server backfill implementation for battlelab/creative.

diff --git a/Engine/Plugins/Online/OnlineSubsystemUtils/Source/OnlineSubsystemUtils/Public/PartyBeaconHost.h b/Engine/Plugins/Online/OnlineSubsystemUtils/Source/OnlineSubsystemUtils/Public/PartyBeaconHost.h
--- a/Engine/Plugins/Online/OnlineSubsystemUtils/Source/OnlineSubsystemUtils/Public/PartyBeaconHost.h
+++ b/Engine/Plugins/Online/OnlineSubsystemUtils/Source/OnlineSubsystemUtils/Public/PartyBeaconHost.h
@@ -330,12 +330,12 @@ class ONLINESUBSYSTEMUTILS_API APartyBeaconHost : public AOnlineBeaconHostObject
 	virtual int32 GetReservationPlatformCount(const FString& InPlatform) const;
 
 	/**
-	 * Delegate fired when a the beacon host detects a reservation addition/removal
+	 * Delegate fired when a the beacon host detects that all reservations are full
 	 */
 	FOnReservationUpdate& OnReservationsFull() { return ReservationsFull; }
 
 	/**
-	 * Delegate fired when a the beacon host detects that all reservations are full
+	 * Delegate fired when a the beacon host detects a reservation addition/removal
 	 */
 	FOnReservationUpdate& OnReservationChanged() { return ReservationChanged; }
 
diff --git a/Engine/Source/Runtime/Json/Public/Serialization/JsonSerializerMacros.h b/Engine/Source/Runtime/Json/Public/Serialization/JsonSerializerMacros.h
--- a/Engine/Source/Runtime/Json/Public/Serialization/JsonSerializerMacros.h
+++ b/Engine/Source/Runtime/Json/Public/Serialization/JsonSerializerMacros.h
@@ -150,8 +150,9 @@
 
 struct FJsonSerializerBase;
 
-/** Array of string data */
+/** Array of data */
 typedef TArray<FString> FJsonSerializableArray;
+typedef TArray<int32> FJsonSerializableArrayInt;
 
 /** Maps a key to a value */
 typedef TMap<FString, FString> FJsonSerializableKeyValueMap;
@@ -182,6 +183,7 @@ struct FJsonSerializerBase
 	virtual void Serialize(const TCHAR* Name, FDateTime& Value) = 0;
 	virtual void SerializeArray(FJsonSerializableArray& Array) = 0;
 	virtual void SerializeArray(const TCHAR* Name, FJsonSerializableArray& Value) = 0;
+	virtual void SerializeArray(const TCHAR* Name, FJsonSerializableArrayInt& Value) = 0;
 	virtual void SerializeMap(const TCHAR* Name, FJsonSerializableKeyValueMap& Map) = 0;
 	virtual void SerializeMap(const TCHAR* Name, FJsonSerializableKeyValueMapInt& Map) = 0;
 	virtual void SerializeMap(const TCHAR* Name, FJsonSerializableKeyValueMapInt64& Map) = 0;
@@ -387,6 +389,22 @@ public:
 		}
 		JsonWriter->WriteArrayEnd();
 	}
+	/**
+	 * Serializes an array of values with an identifier
+	 *
+	 * @param Name the name of the property to serialize
+	 * @param Array the array to serialize
+	 */
+	virtual void SerializeArray(const TCHAR* Name, FJsonSerializableArrayInt& Array) override
+	{
+		JsonWriter->WriteArrayStart(Name);
+		// Iterate all of values
+		for (FJsonSerializableArrayInt::ElementType& Item : Array)
+		{
+			JsonWriter->WriteValue(Item);
+		}
+		JsonWriter->WriteArrayEnd();
+	}
 
 	/**
 	 * Serializes the keys & values for map
@@ -674,6 +692,24 @@ public:
 			}
 		}
 	}
+	/**
+	 * Serializes an array of values with an identifier
+	 *
+	 * @param Name the name of the property to serialize
+	 * @param Array the array to serialize
+	 */
+	virtual void SerializeArray(const TCHAR* Name, FJsonSerializableArrayInt& Array) override
+	{
+		if (JsonObject->HasTypedField<EJson::Array>(Name))
+		{
+			TArray< TSharedPtr<FJsonValue> > JsonArray = JsonObject->GetArrayField(Name);
+			// Iterate all of the keys and their values
+			for (TSharedPtr<FJsonValue>& Value : JsonArray)
+			{
+				Array.Add(Value->AsNumber());
+			}
+		}
+	}
 	/**
 	 * Serializes the keys & values for map
 	 *
