#json Added FJsonValueString::IsEmpty, to allow for avoiding copies

diff --git a/Engine/Source/Runtime/Json/Public/Dom/JsonValue.h b/Engine/Source/Runtime/Json/Public/Dom/JsonValue.h
--- a/Engine/Source/Runtime/Json/Public/Dom/JsonValue.h
+++ b/Engine/Source/Runtime/Json/Public/Dom/JsonValue.h
@@ -121,6 +121,9 @@ public:
 	virtual bool TryGetNumber(uint64& OutValue) const override		{ LexFromString(OutValue, *Value); return true; }
 	virtual bool TryGetBool(bool& OutBool) const override			{ OutBool = Value.ToBool(); return true; }
 
+	// Way to check if string value is empty without copying the string 
+	bool IsEmpty() const { return Value.IsEmpty(); }
+
 protected:
 	FString Value;
 
