Adds 2 new Json Serializer macros.

JSON_SERIALIZE_DATETIME_MILLISECONDS_TIMESTAMP converts to and from a FDateTime to a value in milliseconds (based on the Unix time).

JSON_SERIALIZE_ENUM converts to and from an Enum using LexToString and LexFromString functions for that specific enum.

diff --git a/Engine/Plugins/Online/OnlineSubsystem/Source/Public/OnlineJsonSerializer.h b/Engine/Plugins/Online/OnlineSubsystem/Source/Public/OnlineJsonSerializer.h
--- a/Engine/Plugins/Online/OnlineSubsystem/Source/Public/OnlineJsonSerializer.h
+++ b/Engine/Plugins/Online/OnlineSubsystem/Source/Public/OnlineJsonSerializer.h
@@ -18,6 +18,8 @@
 #define ONLINE_JSON_SERIALIZE_MAP_SERIALIZABLE(JsonName, JsonMap, ElementType)		JSON_SERIALIZE_MAP_SERIALIZABLE(JsonName, JsonMap, ElementType) 
 #define ONLINE_JSON_SERIALIZE_OBJECT_SERIALIZABLE(JsonName, JsonSerializableObject)	JSON_SERIALIZE_OBJECT_SERIALIZABLE(JsonName, JsonSerializableObject) 
 #define ONLINE_JSON_SERIALIZE_DATETIME_UNIX_TIMESTAMP(JsonName, JsonDateTime)		JSON_SERIALIZE_DATETIME_UNIX_TIMESTAMP(JsonName, JsonDateTime)
+#define ONLINE_JSON_SERIALIZE_ENUM(JsonName, JsonEnum)		JSON_SERIALIZE_ENUM(JsonName, JsonEnum)
+#define ONLINE_JSON_SERIALIZE_DATETIME_UNIX_TIMESTAMP_MILLISECONDS(JsonName, JsonDateTime)		JSON_SERIALIZE_DATETIME_UNIX_TIMESTAMP_MILLISECONDS(JsonName, JsonDateTime)
 
 typedef FJsonSerializable FOnlineJsonSerializable;
 typedef FJsonSerializerWriter<> FOnlineJsonSerializerWriter;
diff --git a/Engine/Source/Runtime/Json/Public/Serialization/JsonSerializerMacros.h b/Engine/Source/Runtime/Json/Public/Serialization/JsonSerializerMacros.h
--- a/Engine/Source/Runtime/Json/Public/Serialization/JsonSerializerMacros.h
+++ b/Engine/Source/Runtime/Json/Public/Serialization/JsonSerializerMacros.h
@@ -148,6 +148,33 @@
 			Serializer.Serialize(TEXT(JsonName), UnixTimestampValue); \
 		}
 
+#define JSON_SERIALIZE_DATETIME_UNIX_TIMESTAMP_MILLISECONDS(JsonName, JsonDateTime) \
+if (Serializer.IsLoading()) \
+{ \
+	int64 UnixTimestampValueInMilliseconds; \
+	Serializer.Serialize(TEXT(JsonName), UnixTimestampValueInMilliseconds); \
+	JsonDateTime = FDateTime::FromUnixTimestamp(UnixTimestampValueInMilliseconds / 1000); \
+} \
+else \
+{ \
+	int64 UnixTimestampValueInMilliseconds = JsonDateTime.ToUnixTimestamp() * 1000; \
+	Serializer.Serialize(TEXT(JsonName), UnixTimestampValueInMilliseconds); \
+}
+
+#define JSON_SERIALIZE_ENUM(JsonName, JsonEnum) \
+if (Serializer.IsLoading()) \
+{ \
+	FString JsonTextValue; \
+	Serializer.Serialize(TEXT(JsonName), JsonTextValue); \
+	LexFromString(JsonEnum, *JsonTextValue); \
+} \
+else \
+{ \
+	FString JsonTextValue = LexToString(JsonEnum); \
+	Serializer.Serialize(TEXT(JsonName), JsonTextValue); \
+}
+
+
 struct FJsonSerializerBase;
 
 /** Array of data */
