Added the possibility to serialize FString - float maps using the Json Serializer Macros.

diff --git a/Engine/Source/Runtime/Json/Public/Serialization/JsonSerializerMacros.h b/Engine/Source/Runtime/Json/Public/Serialization/JsonSerializerMacros.h
--- a/Engine/Source/Runtime/Json/Public/Serialization/JsonSerializerMacros.h
+++ b/Engine/Source/Runtime/Json/Public/Serialization/JsonSerializerMacros.h
@@ -185,6 +185,7 @@ typedef TArray<int32> FJsonSerializableArrayInt;
 typedef TMap<FString, FString> FJsonSerializableKeyValueMap;
 typedef TMap<FString, int32> FJsonSerializableKeyValueMapInt;
 typedef TMap<FString, int64> FJsonSerializableKeyValueMapInt64;
+typedef TMap<FString, float> FJsonSerializableKeyValueMapFloat;
 
 /**
  * Base interface used to serialize to/from JSON. Hides the fact there are separate read/write classes
@@ -214,6 +215,7 @@ struct FJsonSerializerBase
 	virtual void SerializeMap(const TCHAR* Name, FJsonSerializableKeyValueMap& Map) = 0;
 	virtual void SerializeMap(const TCHAR* Name, FJsonSerializableKeyValueMapInt& Map) = 0;
 	virtual void SerializeMap(const TCHAR* Name, FJsonSerializableKeyValueMapInt64& Map) = 0;
+	virtual void SerializeMap(const TCHAR* Name, FJsonSerializableKeyValueMapFloat& Map) = 0;
 	virtual void SerializeSimpleMap(FJsonSerializableKeyValueMap& Map) = 0;
 	virtual void SerializeMapSafe(const TCHAR* Name, FJsonSerializableKeyValueMap& Map) = 0;
 	virtual TSharedPtr<FJsonObject> GetObject() = 0;
@@ -484,6 +486,23 @@ public:
 		JsonWriter->WriteObjectEnd();
 	}
 
+	/**
+	 * Serializes the keys & values for map
+	 *
+	 * @param Name the name of the property to serialize
+	 * @param Map the map to serialize
+	 */
+	virtual void SerializeMap(const TCHAR* Name, FJsonSerializableKeyValueMapFloat& Map) override
+	{
+		JsonWriter->WriteObjectStart(Name);
+		// Iterate all of the keys and their values
+		for (FJsonSerializableKeyValueMapFloat::ElementType& Pair : Map)
+		{
+			Serialize(*Pair.Key, Pair.Value);
+		}
+		JsonWriter->WriteObjectEnd();
+	}
+
 	virtual void SerializeSimpleMap(FJsonSerializableKeyValueMap& Map) override
 	{
 		// writing does nothing here, this is meant to read in all data from a json object 
@@ -796,6 +815,26 @@ public:
 		}
 	}
 
+/**
+ * Serializes the keys & values for map
+ *
+ * @param Name the name of the property to serialize
+ * @param Map the map to serialize
+ */
+	virtual void SerializeMap(const TCHAR* Name, FJsonSerializableKeyValueMapFloat& Map) override
+	{
+		if (JsonObject->HasTypedField<EJson::Object>(Name))
+		{
+			TSharedPtr<FJsonObject> JsonMap = JsonObject->GetObjectField(Name);
+			// Iterate all of the keys and their values
+			for (const TPair<FString, TSharedPtr<FJsonValue>>& Pair : JsonMap->Values)
+			{
+				const float Value = (float)Pair.Value->AsNumber();
+				Map.Add(Pair.Key, Value);
+			}
+		}
+	}
+
 	virtual void SerializeSimpleMap(FJsonSerializableKeyValueMap& Map) override
 	{
 		// Iterate all of the keys and their values, only taking simple types (not array/object), all in string form
