Added some includes to fix non-unity compile issues.

diff --git a/Engine/Source/Editor/EditorConfig/Private/EditorConfigSubystem.cpp b/Engine/Source/Editor/EditorConfig/Private/EditorConfigSubystem.cpp
--- a/Engine/Source/Editor/EditorConfig/Private/EditorConfigSubystem.cpp
+++ b/Engine/Source/Editor/EditorConfig/Private/EditorConfigSubystem.cpp
@@ -3,6 +3,8 @@
 #include "EditorConfigSubsystem.h"
 
 #include "Async/Async.h"
+#include "Misc/App.h"
+#include "Misc/Paths.h"
 #include "Misc/ScopeExit.h"
 
 UEditorConfigSubsystem::UEditorConfigSubsystem()
diff --git a/Engine/Source/Editor/EditorConfig/Private/Tests/EditorConfigTestHelpers.h b/Engine/Source/Editor/EditorConfig/Private/Tests/EditorConfigTestHelpers.h
--- a/Engine/Source/Editor/EditorConfig/Private/Tests/EditorConfigTestHelpers.h
+++ b/Engine/Source/Editor/EditorConfig/Private/Tests/EditorConfigTestHelpers.h
@@ -2,6 +2,8 @@
 
 #pragma once
 
+#include "CoreMinimal.h"
+
 namespace FEditorConfigTestHelpers
 {
 	// check that two serialized JSON objects are equivalent
diff --git a/Engine/Source/Editor/EditorConfig/Public/JsonConfig.h b/Engine/Source/Editor/EditorConfig/Public/JsonConfig.h
--- a/Engine/Source/Editor/EditorConfig/Public/JsonConfig.h
+++ b/Engine/Source/Editor/EditorConfig/Public/JsonConfig.h
@@ -4,8 +4,9 @@
 
 #include "CoreMinimal.h"
 #include "Async/Async.h"
-#include "Dom/JsonObject.h"
 #include "Concepts/EqualityComparable.h"
+#include "Containers/StringView.h"
+#include "Dom/JsonObject.h"
 
 namespace UE
 {
diff --git a/Engine/Source/Runtime/Json/Public/Serialization/JsonSerializer.h b/Engine/Source/Runtime/Json/Public/Serialization/JsonSerializer.h
--- a/Engine/Source/Runtime/Json/Public/Serialization/JsonSerializer.h
+++ b/Engine/Source/Runtime/Json/Public/Serialization/JsonSerializer.h
@@ -3,10 +3,10 @@
 #pragma once
 
 #include "CoreMinimal.h"
-#include "Serialization/JsonTypes.h"
-#include "Serialization/JsonReader.h"
 #include "Dom/JsonValue.h"
 #include "Dom/JsonObject.h"
+#include "Serialization/JsonReader.h"
+#include "Serialization/JsonTypes.h"
 #include "Serialization/JsonWriter.h"
 
 class Error;
