Json: Move precision / writing numbers to policy to allow for customizatoin

diff --git a/Engine/Source/Runtime/Json/Public/Policies/JsonPrintPolicy.h b/Engine/Source/Runtime/Json/Public/Policies/JsonPrintPolicy.h
--- a/Engine/Source/Runtime/Json/Public/Policies/JsonPrintPolicy.h
+++ b/Engine/Source/Runtime/Json/Public/Policies/JsonPrintPolicy.h
@@ -38,6 +38,30 @@ struct TJsonPrintPolicy
 			WriteChar(Stream, *CharPtr);
 		}
 	}
+
+	/**
+	 * Writes a float to the output stream.
+	 *
+	 * @param Stream The stream to write to.
+	 * @param Value The float to write.
+	 */
+	static inline void WriteFloat( FArchive* Stream, float Value )
+	{
+		WriteString(Stream, FString::Printf(TEXT("%g"), Value));
+	}
+
+	/**
+	 * Writes a double to the output stream.
+	 *
+	 * @param Stream The stream to write to.
+	 * @param Value The double to write.
+	 */
+	static inline void WriteDouble(  FArchive* Stream, double Value )
+	{
+		// Specify 17 significant digits, the most that can ever be useful from a double
+		// In particular, this ensures large integers are written correctly
+		WriteString(Stream, FString::Printf(TEXT("%.17g"), Value));
+	}
 };
 
 
diff --git a/Engine/Source/Runtime/Json/Public/Serialization/JsonWriter.h b/Engine/Source/Runtime/Json/Public/Serialization/JsonWriter.h
--- a/Engine/Source/Runtime/Json/Public/Serialization/JsonWriter.h
+++ b/Engine/Source/Runtime/Json/Public/Serialization/JsonWriter.h
@@ -360,7 +360,7 @@ protected:
 
 	FORCEINLINE EJsonToken WriteValueOnly(float Value)
 	{
-		PrintPolicy::WriteString(Stream, FString::Printf(TEXT("%g"), Value));
+		PrintPolicy::WriteFloat(Stream, Value);
 		return EJsonToken::Number;
 	}
 
@@ -368,7 +368,7 @@ protected:
 	{
 		// Specify 17 significant digits, the most that can ever be useful from a double
 		// In particular, this ensures large integers are written correctly
-		PrintPolicy::WriteString(Stream, FString::Printf(TEXT("%.17g"), Value));
+		PrintPolicy::WriteDouble(Stream, Value);
 		return EJsonToken::Number;
 	}
 
