Type truncation warning fix in Json header

diff --git a/Engine/Source/Runtime/Json/Public/Serialization/JsonSerializerMacros.h b/Engine/Source/Runtime/Json/Public/Serialization/JsonSerializerMacros.h
--- a/Engine/Source/Runtime/Json/Public/Serialization/JsonSerializerMacros.h
+++ b/Engine/Source/Runtime/Json/Public/Serialization/JsonSerializerMacros.h
@@ -818,7 +818,7 @@ public:
 			// Iterate all of the keys and their values
 			for (TSharedPtr<FJsonValue>& Value : JsonArray)
 			{
-				Array.Add(Value->AsNumber());
+				Array.Add((int32)Value->AsNumber());
 			}
 		}
 	}
