Flipped UnsafeTypeCastWarningLevel on in Json module
#rnx

diff --git a/Engine/Source/Runtime/Json/Json.Build.cs b/Engine/Source/Runtime/Json/Json.Build.cs
--- a/Engine/Source/Runtime/Json/Json.Build.cs
+++ b/Engine/Source/Runtime/Json/Json.Build.cs
@@ -18,6 +18,8 @@ namespace UnrealBuildTool.Rules
 					"Runtime/Json/Private",
 				}
 			);
+
+			UnsafeTypeCastWarningLevel = WarningLevel.Warning;
 		}
 	}
 }
