Handled file I/O errors in JSON deserializer
  - Detected when the file archive failed to read from the file, set the error message to 'Stream I/O Error' and made the deserializer return false.

diff --git a/Engine/Source/Runtime/Json/Public/Serialization/JsonReader.h b/Engine/Source/Runtime/Json/Public/Serialization/JsonReader.h
--- a/Engine/Source/Runtime/Json/Public/Serialization/JsonReader.h
+++ b/Engine/Source/Runtime/Json/Public/Serialization/JsonReader.h
@@ -133,7 +133,7 @@ public:
 
 		if (FinishedReadingRootObject && !Stream->AtEnd())
 		{
-			ParseWhiteSpace();
+			ReadWasSuccess = ParseWhiteSpace();
 		}
 
 		return ReadWasSuccess;
@@ -269,12 +269,15 @@ private:
 			}
 		}
 
-		return true;
+		return !Stream->IsError();
 	}
 
 	bool ReadStart( EJsonToken& Token )
 	{
-		ParseWhiteSpace();
+		if (!ParseWhiteSpace())
+		{
+			return false;
+		}
 
 		Token = EJsonToken::None;
 
@@ -397,7 +400,10 @@ private:
 		while (!Stream->AtEnd())
 		{
 			CharType Char;
-			Stream->Serialize(&Char, sizeof(CharType));
+			if (!Serialize(&Char, sizeof(CharType)))
+			{
+				return false;
+			}
 			++CharacterNumber;
 
 			if (Char == CharType('\0'))
@@ -492,7 +498,10 @@ private:
 
 						while (!Stream->AtEnd())
 						{
-							Stream->Serialize(&Char, sizeof(CharType));
+							if (!Serialize(&Char, sizeof(CharType)))
+							{
+								return false;
+							}
 
 							if (IsAlphaNumber(Char))
 							{
@@ -555,7 +564,10 @@ private:
 			}
 
 			CharType Char;
-			Stream->Serialize(&Char, sizeof(CharType));
+			if (!Serialize(&Char, sizeof(CharType)))
+			{
+				return false;
+			}
 			++CharacterNumber;
 
 			if (Char == CharType('\"'))
@@ -565,7 +577,10 @@ private:
 
 			if (Char == CharType('\\'))
 			{
-				Stream->Serialize(&Char, sizeof(CharType));
+				if (!Serialize(&Char, sizeof(CharType)))
+				{
+					return false;
+				}
 				++CharacterNumber;
 
 				switch (Char)
@@ -589,7 +604,10 @@ private:
 								return false;
 							}
 
-							Stream->Serialize(&Char, sizeof(CharType));
+							if (!Serialize(&Char, sizeof(CharType)))
+							{
+								return false;
+							}
 							++CharacterNumber;
 
 							int32 HexDigit = FParse::HexDigit(Char);
@@ -650,7 +668,10 @@ private:
 			}
 			else
 			{
-				Stream->Serialize(&Char, sizeof(CharType));
+				if (!Serialize(&Char, sizeof(CharType)))
+				{
+					return false;
+				}
 				++CharacterNumber;
 			}
 
@@ -751,12 +772,15 @@ private:
 		return false;
 	}
 
-	void ParseWhiteSpace()
+	bool ParseWhiteSpace()
 	{
 		while (!Stream->AtEnd())
 		{
 			CharType Char;
-			Stream->Serialize(&Char, sizeof(CharType));
+			if (!Serialize(&Char, sizeof(CharType)))
+			{
+				return false;
+			}
 			++CharacterNumber;
 
 			if (IsLineBreak(Char))
@@ -773,6 +797,7 @@ private:
 				break;
 			}
 		}
+		return true;
 	}
 
 	bool IsLineBreak( const CharType& Char )
@@ -810,6 +835,18 @@ private:
 		return (Char >= CharType('a') && Char <= CharType('z')) || (Char >= CharType('A') && Char <= CharType('Z'));
 	}
 
+protected:
+	bool Serialize(void* V, int64 Length)
+	{
+		Stream->Serialize(V, Length);
+		if (Stream->IsError())
+		{
+			SetErrorMessage(TEXT("Stream I/O Error"));
+			return false;
+		}
+		return true;
+	}
+
 protected:
 
 	TArray<EJson> ParseState;
