JsonPointer support:

Array append ("-") unsupported due to general Json impl

diff --git a/Engine/Source/Runtime/Json/Private/JsonUtils/JsonPointer.cpp b/Engine/Source/Runtime/Json/Private/JsonUtils/JsonPointer.cpp
new file mode 100644
--- /dev/null
+++ b/Engine/Source/Runtime/Json/Private/JsonUtils/JsonPointer.cpp
@@ -0,0 +1,207 @@
+﻿// Copyright Epic Games, Inc. All Rights Reserved.
+
+#include "JsonUtils/JsonPointer.h"
+
+#include "Dom/JsonObject.h"
+
+namespace UE::Json
+{
+	namespace Private
+	{
+		static constexpr TCHAR RootCharacter = TEXT('#');
+		static constexpr TCHAR PathDelimiterCharacter = TEXT('/');
+		static constexpr TCHAR ArrayAppendCharacter = TEXT('-'); // unused, but kept for reference as official part of spec
+	}
+
+	FJsonPointer::FJsonPointer(FStringView InPath)
+	{
+		ParsePath(FString(InPath));
+	}
+
+	FJsonPointer::FJsonPointer(const FJsonPointer& Other)
+		: PathParts(Other.PathParts)
+	{
+	}
+
+	FJsonPointer::FJsonPointer(FJsonPointer&& Other) noexcept
+		: PathParts(MoveTemp(Other.PathParts))
+	{
+	}
+
+	static const TSharedPtr<FJsonValue>* FindPartInJsonValue(const TSharedPtr<FJsonValue>& Value, const FString& PathPart)
+	{
+		const TSharedPtr<FJsonValue>* NextValue = &Value;
+
+		TArray<TSharedPtr<FJsonValue>>* ArrayPtr = nullptr;
+		if ((*NextValue)->TryGetArray(ArrayPtr))
+		{
+			if(!PathPart.IsNumeric())
+			{
+				// part expected to be index
+				UE_LOG(LogJson, Warning, TEXT("FJsonPointer::TryGet: JsonValue was an array and expected a numeric part but was a string."));
+				return nullptr;
+			}
+			
+			const int32 PartAsIndex = FCString::Atoi(*PathPart);		
+			if (!ArrayPtr->IsValidIndex(PartAsIndex))
+			{
+				// index out of range
+				UE_LOG(LogJson, Warning, TEXT("FJsonPointer::TryGet: Array index %d for array is out of range."), PartAsIndex);
+				return nullptr;
+			}
+
+			NextValue = &(*ArrayPtr)[PartAsIndex];
+		}
+		else
+		{
+			const TSharedPtr<FJsonObject>* ObjectPtr;
+			if((*NextValue)->TryGetObject(ObjectPtr))
+			{
+				NextValue = (*ObjectPtr)->Values.Find(PathPart);
+				if (NextValue == nullptr)
+				{
+					// part of path didn't exist
+					UE_LOG(LogJson, Warning, TEXT("FJsonPointer::TryGet: Field \"%s\" not found."), *PathPart);
+					return nullptr;
+				}
+			}
+		}
+
+		return NextValue;
+	}
+
+	bool FJsonPointer::TryGet(const TSharedPtr<FJsonObject>& RootObject, TSharedPtr<FJsonValue>& OutValue) const
+	{			
+		return TryGet(MakeShared<FJsonValueObject>(RootObject), OutValue);
+	}
+
+	bool FJsonPointer::TryGet(const TSharedPtr<FJsonValue>& RootValue, TSharedPtr<FJsonValue>& OutValue) const
+	{
+		if (!IsValid())
+		{
+			return false;
+		}
+
+		const TSharedPtr<FJsonValue>* CurrentValue = &RootValue;
+
+		for (int32 Idx = 0; Idx < Length(); ++Idx)
+		{
+			CurrentValue = FindPartInJsonValue(*CurrentValue, PathParts[Idx]);
+			if (CurrentValue == nullptr)
+			{
+				return false;
+			}
+		}
+
+		OutValue = *CurrentValue;
+		return true;
+	}
+
+	bool FJsonPointer::TryGet(const TArray<TSharedPtr<FJsonValue>>& RootValues, TSharedPtr<FJsonValue>& OutValue) const
+	{
+		if (!IsValid())
+		{
+			return false;
+		}
+
+		const FString& FirstPart = PathParts[0];
+		if(!FirstPart.IsNumeric())
+		{
+			UE_LOG(LogJson, Warning, TEXT("FJsonPointer::TryGet: RootObject was an array, but JsonPointer path requested a field \"%s\"."), *FirstPart);
+			return false;
+		}
+
+		const int32 PartAsIndex = FCString::Atoi(*FirstPart);
+		if (!RootValues.IsValidIndex(PartAsIndex))
+		{
+			// index out of range
+			UE_LOG(LogJson, Warning, TEXT("FJsonPointer::TryGet: Array index %d for array %s is out of range."), PartAsIndex, *FirstPart);
+			return false;
+		}
+
+		return TryGet(RootValues[PartAsIndex], OutValue);
+	}
+
+	bool FJsonPointer::ParsePath(const FString& InPath)
+	{
+		TArray<FString> PathStrings;
+		InPath.ParseIntoArray(PathStrings, &Private::PathDelimiterCharacter, false);
+
+		if(InPath.IsEmpty())
+		{
+			return false;
+		}
+
+		// We don't really care about the document root, remove it if it exists
+		if (PathStrings[0] == &Private::RootCharacter)
+		{
+			PathStrings.RemoveAt(0);
+		}
+
+		for (int32 Idx = 0; Idx < PathStrings.Num(); ++Idx)
+		{
+			FString PartString = PathStrings[Idx];
+			if (PartString.IsEmpty())
+			{
+				PathParts.Empty();
+				UE_LOG(LogJson, Warning, TEXT("Path part was empty in JSON pointer \"%s\""), *InPath);
+				return false;
+			}
+
+			PathParts.Add(PartString);
+		}
+
+		return true;
+	}
+
+	void FJsonPointer::Append(FStringView Part)
+	{
+		PathParts.Add(FString(Part));
+	}
+
+	void FJsonPointer::SetArrayIndex(int32 Index)
+	{
+		check(PathParts.Num() > 0);
+
+		PathParts.Last() = LexToString(Index);
+	}
+
+	FJsonPointer FJsonPointer::GetSubPath(int32 NumParts) const
+	{
+		// NumParts is or exceeds length, so just return the whole path
+		if (NumParts >= Length())
+		{
+			return *this;
+		}
+
+		FJsonPointer SubPath;
+		SubPath.PathParts.Reserve(NumParts);
+
+		for (int32 Idx = 0; Idx < NumParts; ++Idx)
+		{
+			SubPath.PathParts.Add(PathParts[Idx]);
+		}
+
+		return SubPath;
+	}
+
+	FString FJsonPointer::ToString() const
+	{
+		return Private::RootCharacter + Private::PathDelimiterCharacter + *FString::Join(PathParts, &Private::PathDelimiterCharacter);
+	}
+
+	FString FJsonPointer::EscapePart(const FString& Part)
+	{
+		return Part
+			.Replace(TEXT("~"), TEXT("~0"))
+			.Replace(TEXT("/"), TEXT("~1"));
+	}
+
+	FString FJsonPointer::UnescapePart(const FString& Part)
+	{
+		// Performed in order according to: https://datatracker.ietf.org/doc/html/rfc6901#section-4
+		return Part
+			.Replace(TEXT("~1"), TEXT("/"))
+			.Replace(TEXT("~0"), TEXT("~"));
+	}
+}
diff --git a/Engine/Source/Runtime/Json/Private/Tests/JsonPointer.spec.cpp b/Engine/Source/Runtime/Json/Private/Tests/JsonPointer.spec.cpp
new file mode 100644
--- /dev/null
+++ b/Engine/Source/Runtime/Json/Private/Tests/JsonPointer.spec.cpp
@@ -0,0 +1,264 @@
+﻿// Copyright Epic Games, Inc. All Rights Reserved.
+
+#include "JsonUtils/JsonPointer.h"
+#include "Misc/AutomationTest.h"
+#include "Serialization/JsonReader.h"
+#include "Serialization/JsonSerializer.h"
+
+#if WITH_AUTOMATION_TESTS
+
+BEGIN_DEFINE_SPEC(FJsonPointerSpec, TEXT("JsonPointer"), EAutomationTestFlags::EngineFilter | EAutomationTestFlags::ApplicationContextMask)
+TSharedPtr<FJsonObject> JsonObject;
+END_DEFINE_SPEC(FJsonPointerSpec)
+
+void FJsonPointerSpec::Define()
+{
+	BeforeEach([this]
+	{
+		static constexpr const TCHAR* ExampleJson = TEXT(R"_JSON(
+		{
+		    "someNumber": 42,
+			"someObject": {
+				"someNestedNumber": 24
+			},
+			"someArray": [
+				5,
+				3,
+				22,
+				2,
+				44
+			]
+		})_JSON");
+
+		const TSharedRef<TJsonReader<>> JsonReader = TJsonReaderFactory<>::Create(ExampleJson);
+		return FJsonSerializer::Deserialize(JsonReader, JsonObject);
+	});
+
+	Describe(TEXT("Value"), [this]
+	{
+		It(TEXT("Existing value is found and returned"), [this]
+		{
+			const UE::Json::FJsonPointer JsonPointer{TEXT("#/someNumber"});
+			
+			TSharedPtr<FJsonValue> JsonValue;
+			const bool bFound = JsonPointer.TryGet(JsonObject, JsonValue);
+
+			constexpr int32 ExpectedValue = 42;
+			
+			if(TestTrue(TEXT("Value is found"), bFound))
+			{
+				TestEqual(TEXT("Value equals 42"), StaticCast<int32>(JsonValue->AsNumber()), ExpectedValue);
+			}
+		});
+
+		It(TEXT("Existing nested value is found and returned"), [this]
+		{
+			const UE::Json::FJsonPointer JsonPointer{TEXT("#/someObject/someNestedNumber"});
+						
+			TSharedPtr<FJsonValue> JsonValue;
+			const bool bFound = JsonPointer.TryGet(JsonObject, JsonValue);
+
+			constexpr int32 ExpectedValue = 24;
+						
+			if(TestTrue(TEXT("Value is found"), bFound))
+			{
+				TestEqual(TEXT("Value equals 42"), StaticCast<int32>(JsonValue->AsNumber()), ExpectedValue);
+			}
+		});
+		
+		It(TEXT("Non-existent value logs a warning"), [this]
+		{
+			const UE::Json::FJsonPointer JsonPointer{TEXT("#/foo/doesnt/exist"});
+
+			AddExpectedError(TEXT("not found"), EAutomationExpectedErrorFlags::Contains);
+						
+			TSharedPtr<FJsonValue> JsonValue;
+			const bool bFound = JsonPointer.TryGet(JsonObject, JsonValue);
+
+			TestFalse(TEXT("Value is not found"), bFound);
+		});
+
+		Describe(TEXT("Values with numeric field names are distinguised from array values"), [this]
+		{
+			BeforeEach([this]
+			{
+				JsonObject.Reset();
+				
+				// Overrides what happens in outer BeforeEach
+				static constexpr const TCHAR* ExampleJson = TEXT(R"_JSON(
+				{
+				    "0": "foo",
+					"nested": {
+						"1": "bar"
+					},
+					"2": {
+						"bla": false
+					},
+					"someArray": [
+					{
+						"Foo": "foo",
+						"Bar": "bar"
+					}]
+				})_JSON");
+
+				const TSharedRef<TJsonReader<>> JsonReader = TJsonReaderFactory<>::Create(ExampleJson);
+				return FJsonSerializer::Deserialize(JsonReader, JsonObject);
+			});
+
+			It(TEXT("Existing value with numeric field name \"0\" is found and returned as a value, rather than an item in an array"), [this]
+			{
+				const UE::Json::FJsonPointer JsonPointer{TEXT("#/0"});
+				
+				TSharedPtr<FJsonValue> JsonValue;
+				const bool bFound = JsonPointer.TryGet(JsonObject, JsonValue);
+
+				const FString ExpectedValue = TEXT("foo");
+
+				if(TestTrue(TEXT("Value is found"), bFound))
+				{
+					TestEqual(TEXT("Value equals \"foo\""), JsonValue->AsString(), ExpectedValue);
+				}
+			});
+
+			It(TEXT("Existing nested value with numeric field name \"1\" is found and returned as a value, rather than an item in an array"), [this]
+			{
+				const UE::Json::FJsonPointer JsonPointer{TEXT("#/nested/1"});
+				
+				TSharedPtr<FJsonValue> JsonValue;
+				const bool bFound = JsonPointer.TryGet(JsonObject, JsonValue);
+
+				const FString ExpectedValue = TEXT("bar");
+
+				if(TestTrue(TEXT("Value is found"), bFound))
+				{
+					TestEqual(TEXT("Value equals \"bar\""), JsonValue->AsString(), ExpectedValue);
+				}
+			});
+
+			It(TEXT("Existing nested value where the parent has a numeric field name \"2\" is found and returned"), [this]
+			{
+				const UE::Json::FJsonPointer JsonPointer{TEXT("#/2/bla"});
+									
+				TSharedPtr<FJsonValue> JsonValue;
+				const bool bFound = JsonPointer.TryGet(JsonObject, JsonValue);
+
+				constexpr bool ExpectedValue = false;
+
+				if(TestTrue(TEXT("Value is found"), bFound))
+				{
+					TestEqual(TEXT("Value equals false"), JsonValue->AsBool(), ExpectedValue);
+				}
+			});
+
+			It(TEXT("Existing value within an object, within an array is found and returned"), [this]
+			{
+				const UE::Json::FJsonPointer JsonPointer{TEXT("#/someArray/0/Bar"});
+												
+				TSharedPtr<FJsonValue> JsonValue;
+				const bool bFound = JsonPointer.TryGet(JsonObject, JsonValue);
+
+				const FString ExpectedValue = TEXT("bar");
+
+				if(TestTrue(TEXT("Value is found"), bFound))
+				{
+					TestEqual(TEXT("Value equals false"), JsonValue->AsString(), ExpectedValue);
+				}
+			});
+		});
+	});
+
+	Describe(TEXT("Array Value"), [this]
+	{
+		It(TEXT("Existing value is found and returned at index 2"), [this]
+		{
+			const UE::Json::FJsonPointer JsonPointer{TEXT("#/someArray/2"});
+			
+			TSharedPtr<FJsonValue> JsonValue;
+			const bool bFound = JsonPointer.TryGet(JsonObject, JsonValue);
+
+			constexpr int32 ExpectedValue = 22;
+			
+			if(TestTrue(TEXT("Value is found"), bFound))
+			{
+				TestEqual(TEXT("Value equals 22"), StaticCast<int32>(JsonValue->AsNumber()), ExpectedValue);
+			}
+		});
+
+		It(TEXT("Out of range index 5 logs a warning"), [this]
+		{
+			const UE::Json::FJsonPointer JsonPointer{TEXT("#/someArray/5"});
+
+			AddExpectedError(TEXT("out of range"), EAutomationExpectedErrorFlags::Contains);
+			
+			TSharedPtr<FJsonValue> JsonValue;
+			const bool bFound = JsonPointer.TryGet(JsonObject, JsonValue);
+
+			TestFalse(TEXT("Value is not found"), bFound);
+		});
+
+		// When the document root is an array
+		It(TEXT("Existing value in root is found and returned at index 2"), [this]
+		{
+			// Overrides what happens in BeforeEach
+			static constexpr const TCHAR* ExampleJson = TEXT("[ 5, 9, 7, 2 ]");
+
+			TArray<TSharedPtr<FJsonValue>> JsonObjects;
+			const TSharedRef<TJsonReader<>> JsonReader = TJsonReaderFactory<>::Create(ExampleJson);			
+			if(TestTrue(TEXT("Deserialized when root is array"), FJsonSerializer::Deserialize(JsonReader, JsonObjects)))
+			{
+				const UE::Json::FJsonPointer JsonPointer{TEXT("#/2"});
+					
+				TSharedPtr<FJsonValue> JsonValue;
+				const bool bFound = JsonPointer.TryGet(JsonObjects, JsonValue);
+
+				constexpr int32 ExpectedValue = 7;
+
+				if(TestTrue(TEXT("Value is found"), bFound))
+				{
+					TestEqual(TEXT("Value equals 7"), StaticCast<int32>(JsonValue->AsNumber()), ExpectedValue);
+				}
+			}
+		});
+
+		It(TEXT("Path part immediately following an array field must be an index"), [this]
+		{
+			const UE::Json::FJsonPointer JsonPointer{TEXT("#/someArray/Bar")};
+			
+			AddExpectedError(TEXT("expected a numeric part"));
+			
+			TSharedPtr<FJsonValue> JsonValue;
+			const bool bFound = JsonPointer.TryGet(JsonObject, JsonValue);
+
+			const FString ExpectedValue = TEXT("bar");
+
+			TestFalse(TEXT("Value is found"), bFound);
+		});
+	});
+
+	Describe(TEXT("Syntax"), [this]
+	{
+		// Tests https://datatracker.ietf.org/doc/html/rfc6901#section-4
+		It(TEXT("Escapes ~ and /"), [this]
+		{
+			static constexpr const TCHAR* PartToTest = TEXT(R"-(Some~part/with~escapable//Characters)-");
+			static constexpr const TCHAR* PartToExpect = TEXT(R"-(Some~0part~1with~0escapable~1~1Characters)-");
+
+			const FString EscapedPart = UE::Json::FJsonPointer::EscapePart(PartToTest);
+
+			TestEqual(TEXT("\"~\" and \"/\" are replaced with \"~0\" and \"~1\" respectively"), EscapedPart, PartToExpect);
+		});
+
+		// Tests https://datatracker.ietf.org/doc/html/rfc6901#section-4
+		It(TEXT("Unescapes ~ and /"), [this]
+		{
+			static constexpr const TCHAR* PartToTest = TEXT(R"-(Some~0part~1with~0escapable~1~1Characters)-");
+			static constexpr const TCHAR* PartToExpect = TEXT(R"-(Some~part/with~escapable//Characters)-");
+
+			const FString UnescapedPart = UE::Json::FJsonPointer::UnescapePart(PartToTest);
+
+			TestEqual(TEXT("\"~0\" and \"~1\" are replaced with \"~\" and \"/\" respectively"), UnescapedPart, PartToExpect);
+		});
+	});
+}
+
+#endif
diff --git a/Engine/Source/Runtime/Json/Public/JsonUtils/JsonPointer.h b/Engine/Source/Runtime/Json/Public/JsonUtils/JsonPointer.h
new file mode 100644
--- /dev/null
+++ b/Engine/Source/Runtime/Json/Public/JsonUtils/JsonPointer.h
@@ -0,0 +1,82 @@
+﻿// Copyright Epic Games, Inc. All Rights Reserved.
+
+#pragma once
+
+#include "CoreMinimal.h"
+#include "Dom/JsonValue.h"
+
+namespace UE::Json
+{
+	/**
+	* JSON Pointer defines a string syntax for identifying a specific value
+	* within a JavaScript Object Notation (JSON) document (from rfc6901).
+	* @see https://datatracker.ietf.org/doc/html/rfc6901
+	* @note Array append ("-") is currently unsupported.
+	*/
+	class JSON_API FJsonPointer
+	{
+	public:
+		FJsonPointer() = default;
+		FJsonPointer(FStringView Path);
+		FJsonPointer(const FJsonPointer& Other);
+		FJsonPointer(FJsonPointer&& Other) noexcept;
+
+		/** Attempt to retrieve the value specified by this Json Pointer from the provided Json Value. */
+		bool TryGet(const TSharedPtr<FJsonObject>& RootObject, TSharedPtr<FJsonValue>& OutValue) const;
+
+		/** Attempt to retrieve the value specified by this Json Pointer from the provided Json Object. */
+		bool TryGet(const TSharedPtr<FJsonValue>& RootValue, TSharedPtr<FJsonValue>& OutValue) const;
+
+		/** Attempt to retrieve the value specified by this Json Pointer from the provided Json Array. */
+		bool TryGet(const TArray<TSharedPtr<FJsonValue>>& RootValues, TSharedPtr<FJsonValue>& OutValue) const;
+
+		/** Checks that the path has one or more parts. */
+		FORCEINLINE bool IsValid() const
+		{
+			return Length() > 0;
+		}
+
+		/** The number of parts in the path. */
+		FORCEINLINE int32 Length() const
+		{
+			return PathParts.Num();
+		}
+
+		/** Appends the given part to the path */
+		void Append(FStringView Part);
+
+		/** Sets the last path part to the given array index.  */
+		void SetArrayIndex(int32 Index);
+
+		/** Gets the sub-path up to the given part length. */
+		FJsonPointer GetSubPath(int32 NumParts) const;
+
+		/** Get the JsonPointer's string representation. */
+		FString ToString() const;
+
+		/** Get the path part at the given index. */
+		FORCEINLINE const FString& operator[](int32 Idx) const
+		{
+			return PathParts[Idx];
+		}
+
+		/** Get the path parts as an array. */
+		FORCEINLINE const TArray<FString>& GetAll() const
+		{
+			return PathParts;
+		}
+
+		/** Substitutes characters in the provided Part with their escaped counterparts (as per the spec). */
+		static FString EscapePart(const FString& Part);
+		
+		/** Substitutes escaped characters in the Provided path with their original character (as per the spec). */
+		static FString UnescapePart(const FString& Part);
+
+	private:
+		/** Parses the provided path, populating PathParts and validating the result. */
+		bool ParsePath(const FString& InPath);
+
+	private:
+		TArray<FString> PathParts;
+	};
+}
