Prepared the Social Ban Service to handle the possibility of banning features separately.

Made changes to handle this feature in the VCM and in the DS to avoid sending join tokens.

Added new macro to serialize enum arrays.

diff --git a/Engine/Source/Runtime/Json/Public/Serialization/JsonSerializerMacros.h b/Engine/Source/Runtime/Json/Public/Serialization/JsonSerializerMacros.h
--- a/Engine/Source/Runtime/Json/Public/Serialization/JsonSerializerMacros.h
+++ b/Engine/Source/Runtime/Json/Public/Serialization/JsonSerializerMacros.h
@@ -240,6 +240,29 @@ else \
 	Serializer.Serialize(TEXT(JsonName), JsonTextValue); \
 }
 
+#define JSON_SERIALIZE_ENUM_ARRAY(JsonName, JsonArray, EnumType) \
+if (Serializer.IsLoading()) \
+{ \
+	if (Serializer.GetObject()->HasTypedField<EJson::Array>(JsonName)) \
+	{ \
+		EnumType EnumValue; \
+		for (auto It = Serializer.GetObject()->GetArrayField(JsonName).CreateConstIterator(); It; ++It) \
+		{ \
+			LexFromString(EnumValue, *(*It)->AsString()); \
+			JsonArray.Add(EnumValue); \
+		} \
+	} \
+} \
+else \
+{ \
+	Serializer.StartArray(JsonName); \
+	for (EnumType& EnumValue : JsonArray) \
+	{ \
+		FString JsonTextValue = LexToString(EnumValue); \
+		Serializer.Serialize(TEXT(JsonName), JsonTextValue); \
+	} \
+	Serializer.EndArray(); \
+}
 
 struct FJsonSerializerBase;
 
