PR #9568: - Hande empty JSON keys
(Contributed by [REDACTED_PERSON])

diff --git a/Engine/Source/Runtime/Json/Private/Tests/JsonTests.cpp b/Engine/Source/Runtime/Json/Private/Tests/JsonTests.cpp
--- a/Engine/Source/Runtime/Json/Private/Tests/JsonTests.cpp
+++ b/Engine/Source/Runtime/Json/Private/Tests/JsonTests.cpp
@@ -234,6 +234,7 @@ bool FJsonAutomationTest::RunTest(const FString& Parameters)
 					"\"Some String1\","
 					"10,"
 					"{"
+						"\"\":\"Empty Key\","
 						"\"Value\":\"Some String3\""
 					"},"
 					"["
@@ -267,6 +268,7 @@ bool FJsonAutomationTest::RunTest(const FString& Parameters)
 		TSharedPtr< FJsonObject > Object = Array[2]->AsObject();
 		check(Object.IsValid());
 		check(Object->GetStringField(TEXT("Value")) == TEXT("Some String3"));
+		check(Object->GetStringField(TEXT("")) == TEXT("Empty Key"));
 
 		const TArray<TSharedPtr< FJsonValue >>& InnerArray = Array[3]->AsArray();
 		check(InnerArray.Num() == 2);
diff --git a/Engine/Source/Runtime/Json/Public/Serialization/JsonSerializer.h b/Engine/Source/Runtime/Json/Public/Serialization/JsonSerializer.h
--- a/Engine/Source/Runtime/Json/Public/Serialization/JsonSerializer.h
+++ b/Engine/Source/Runtime/Json/Public/Serialization/JsonSerializer.h
@@ -155,35 +155,31 @@ private:
 
 	struct FElement
 	{
-		FElement( const TSharedPtr<FJsonValue>& InValue )
+		FElement(const TSharedPtr<FJsonValue>& InValue)
 			: Identifier()
 			, Value(InValue)
-			, HasBeenProcessed(false)
 		{ }
 
-		FElement( const TSharedRef<FJsonObject>& Object )
+		FElement(const TSharedRef<FJsonObject>& Object)
 			: Identifier()
 			, Value(MakeShared<FJsonValueObject>(Object))
-			, HasBeenProcessed( false )
 		{ }
 
-		FElement( const TArray<TSharedPtr<FJsonValue>>& Array )
+		FElement(const TArray<TSharedPtr<FJsonValue>>& Array)
 			: Identifier()
 			, Value(MakeShared<FJsonValueArray>(Array))
-			, HasBeenProcessed(false)
 		{ }
 
-		FElement( const FString& InIdentifier, const TSharedPtr< FJsonValue >& InValue )
+		FElement(const FString& InIdentifier, const TSharedPtr< FJsonValue >& InValue)
 			: Identifier( InIdentifier )
 			, Value( InValue )
-			, HasBeenProcessed( false )
-		{
-
-		}
+			, bIsKeyValuePair(true)
+		{ }
 
 		FString Identifier;
 		TSharedPtr< FJsonValue > Value;
-		bool HasBeenProcessed;
+		bool bHasBeenProcessed = false;
+		bool bIsKeyValuePair = false;
 	};
 
 private:
@@ -319,74 +315,74 @@ private:
 			{
 			case EJson::Number:	
 				{
-					if (Element->Identifier.IsEmpty())
+					if (Element->bIsKeyValuePair)
 					{
-						Writer.WriteValue(Element->Value->AsNumber());
+						Writer.WriteValue(Element->Identifier, Element->Value->AsNumber());
 					}
 					else
 					{
-						Writer.WriteValue(Element->Identifier, Element->Value->AsNumber());
+						Writer.WriteValue(Element->Value->AsNumber());
 					}
 				}
 				break;
 
 			case EJson::Boolean:					
 				{
-					if (Element->Identifier.IsEmpty())
+					if (Element->bIsKeyValuePair)
 					{
-						Writer.WriteValue(Element->Value->AsBool());
+						Writer.WriteValue(Element->Identifier, Element->Value->AsBool());
 					}
 					else
 					{
-						Writer.WriteValue(Element->Identifier, Element->Value->AsBool());
+						Writer.WriteValue(Element->Value->AsBool());
 					}
 				}
 				break;
 
 			case EJson::String:
 				{
-					if (Element->Identifier.IsEmpty())
+					if (Element->bIsKeyValuePair)
 					{
-						Writer.WriteValue(Element->Value->AsString());
+						Writer.WriteValue(Element->Identifier, Element->Value->AsString());
 					}
 					else
 					{
-						Writer.WriteValue(Element->Identifier, Element->Value->AsString());
+						Writer.WriteValue(Element->Value->AsString());
 					}
 				}
 				break;
 
 			case EJson::Null:
 				{
-					if (Element->Identifier.IsEmpty())
+					if (Element->bIsKeyValuePair)
 					{
-						Writer.WriteNull();
+						Writer.WriteNull(Element->Identifier);
 					}
 					else
 					{
-						Writer.WriteNull(Element->Identifier);
+						Writer.WriteNull();
 					}
 				}
 				break;
 
 			case EJson::Array:
 				{
-					if (Element->HasBeenProcessed)
+					if (Element->bHasBeenProcessed)
 					{
 						Writer.WriteArrayEnd();
 					}
 					else
 					{
-						Element->HasBeenProcessed = true;
+						Element->bHasBeenProcessed = true;
 						ElementStack.Push(Element);
 
-						if (Element->Identifier.IsEmpty())
+						if (Element->bIsKeyValuePair)
 						{
-							Writer.WriteArrayStart();
+							Writer.WriteArrayStart(Element->Identifier);
 						}
 						else
 						{
-							Writer.WriteArrayStart(Element->Identifier);
+							Writer.WriteArrayStart();
 						}
 
 						TArray<TSharedPtr<FJsonValue>> Values = Element->Value->AsArray();
@@ -401,22 +397,22 @@ private:
 
 			case EJson::Object:
 				{
-					if (Element->HasBeenProcessed)
+					if (Element->bHasBeenProcessed)
 					{
 						Writer.WriteObjectEnd();
 					}
 					else
 					{
-						Element->HasBeenProcessed = true;
+						Element->bHasBeenProcessed = true;
 						ElementStack.Push(Element);
 
-						if (Element->Identifier.IsEmpty())
+						if (Element->bIsKeyValuePair)
 						{
-							Writer.WriteObjectStart();
+							Writer.WriteObjectStart(Element->Identifier);
 						}
 						else
 						{
-							Writer.WriteObjectStart(Element->Identifier);
+							Writer.WriteObjectStart();
 						}
 
 						TArray<FString> Keys; 
