Version of FJsonObject::Duplicate that takes in a const FJsonObject as the source

diff --git a/Engine/Source/Runtime/Json/Private/Dom/JsonObject.cpp b/Engine/Source/Runtime/Json/Private/Dom/JsonObject.cpp
--- a/Engine/Source/Runtime/Json/Private/Dom/JsonObject.cpp
+++ b/Engine/Source/Runtime/Json/Private/Dom/JsonObject.cpp
@@ -209,10 +209,18 @@ void FJsonObject::SetObjectField( const FString& FieldName, const TSharedPtr<FJs
 	}
 }
 
-void FJsonObject::Duplicate(const TSharedPtr<FJsonObject>& Source, TSharedPtr<FJsonObject>& Dest)
+void FJsonObject::Duplicate(const TSharedPtr<const FJsonObject>& Source, const TSharedPtr<FJsonObject>& Dest)
 {
-	for (const TPair<FString, TSharedPtr<FJsonValue>>& Pair : Source->Values)
+	if (Source && Dest)
 	{
-		Dest->SetField(Pair.Key, FJsonValue::Duplicate(Pair.Value));
+		for (const TPair<FString, TSharedPtr<FJsonValue>>& Pair : Source->Values)
+		{
+			Dest->SetField(Pair.Key, FJsonValue::Duplicate(Pair.Value));
+		}
 	}
 }
+
+void FJsonObject::Duplicate(const TSharedPtr<FJsonObject>& Source, TSharedPtr<FJsonObject>& Dest)
+{
+	Duplicate(ConstCastSharedPtr<const FJsonObject>(Source), Dest);
+}
diff --git a/Engine/Source/Runtime/Json/Public/Dom/JsonObject.h b/Engine/Source/Runtime/Json/Public/Dom/JsonObject.h
--- a/Engine/Source/Runtime/Json/Public/Dom/JsonObject.h
+++ b/Engine/Source/Runtime/Json/Public/Dom/JsonObject.h
@@ -237,5 +237,6 @@ public:
 	/** Set an ObjectField named FieldName and value of JsonObject */
 	void SetObjectField( const FString& FieldName, const TSharedPtr<FJsonObject>& JsonObject );
 
+	static void Duplicate(const TSharedPtr<const FJsonObject>& Source, const TSharedPtr<FJsonObject>& Dest);
 	static void Duplicate(const TSharedPtr<FJsonObject>& Source, TSharedPtr<FJsonObject>& Dest);
 };
