Add move constructors for JSon value types.

diff --git a/Engine/Source/Runtime/Json/Public/Dom/JsonValue.h b/Engine/Source/Runtime/Json/Public/Dom/JsonValue.h
--- a/Engine/Source/Runtime/Json/Public/Dom/JsonValue.h
+++ b/Engine/Source/Runtime/Json/Public/Dom/JsonValue.h
@@ -125,6 +125,7 @@ class JSON_API FJsonValueString : public FJsonValue
 {
 public:
 	FJsonValueString(const FString& InString) : Value(InString) {Type = EJson::String;}
+	FJsonValueString(FString&& InString) : Value(MoveTemp(InString)) {Type = EJson::String;}
 
 	virtual bool TryGetString(FString& OutString) const override	{ OutString = Value; return true; }
 	virtual bool TryGetNumber(double& OutDouble) const override		{ if (Value.IsNumeric()) { OutDouble = FCString::Atod(*Value); return true; } else { return false; } }
@@ -166,6 +167,7 @@ class JSON_API FJsonValueNumberString : public FJsonValue
 {
 public:
 	FJsonValueNumberString(const FString& InString) : Value(InString) { Type = EJson::Number; }
+	FJsonValueNumberString(FString&& InString) : Value(MoveTemp(InString)) { Type = EJson::Number; }
 
 	virtual bool TryGetString(FString& OutString) const override { OutString = Value; return true; }
 	virtual bool TryGetNumber(double& OutDouble) const override { return LexTryParseString(OutDouble, *Value); }
@@ -208,6 +210,7 @@ class JSON_API FJsonValueArray : public FJsonValue
 {
 public:
 	FJsonValueArray(const TArray< TSharedPtr<FJsonValue> >& InArray) : Value(InArray) {Type = EJson::Array;}
+	FJsonValueArray(TArray< TSharedPtr<FJsonValue> >&& InArray) : Value(MoveTemp(InArray)) {Type = EJson::Array;}
 	virtual bool TryGetArray(const TArray< TSharedPtr<FJsonValue> >*& OutArray) const override	{ OutArray = &Value; return true; }
 	virtual bool TryGetArray(TArray< TSharedPtr<FJsonValue> >*& OutArray) override				{ OutArray = &Value; return true; }
 	
@@ -222,7 +225,7 @@ protected:
 class JSON_API FJsonValueObject : public FJsonValue
 {
 public:
-	FJsonValueObject(TSharedPtr<FJsonObject> InObject) : Value(InObject) {Type = EJson::Object;}
+	FJsonValueObject(TSharedPtr<FJsonObject> InObject) : Value(MoveTemp(InObject)) {Type = EJson::Object;}
 	virtual bool TryGetObject(const TSharedPtr<FJsonObject>*& OutObject) const override	{ OutObject = &Value; return true; }
 	virtual bool TryGetObject(TSharedPtr<FJsonObject>*& OutObject) override				{ OutObject = &Value; return true; }
 	
