Refactor JSON printer to use string views where possible

diff --git a/Engine/Plugins/Enterprise/GLTFExporter/Source/GLTFExporter/Private/Json/GLTFJsonWriter.cpp b/Engine/Plugins/Enterprise/GLTFExporter/Source/GLTFExporter/Private/Json/GLTFJsonWriter.cpp
--- a/Engine/Plugins/Enterprise/GLTFExporter/Source/GLTFExporter/Private/Json/GLTFJsonWriter.cpp
+++ b/Engine/Plugins/Enterprise/GLTFExporter/Source/GLTFExporter/Private/Json/GLTFJsonWriter.cpp
@@ -168,6 +168,26 @@ struct TGLTFJsonPrintPolicyUTF8
 		}
 	}
 
+	static inline void WriteString(FArchive* Stream, FStringView String)
+	{
+		const TCHAR* CharPtr = String.GetData();
+
+		for (int32 CharIndex = 0; CharIndex < String.Len(); ++CharIndex, ++CharPtr)
+		{
+			WriteChar(Stream, *CharPtr);
+		}
+	}
+
+	static inline void WriteString(FArchive* Stream, FAnsiStringView String)
+	{
+		const ANSICHAR* CharPtr = String.GetData();
+
+		for (int32 CharIndex = 0; CharIndex < String.Len(); ++CharIndex, ++CharPtr)
+		{
+			WriteChar(Stream, *CharPtr);
+		}
+	}
+
 	static void WriteFloat(FArchive* Stream, float Number)
 	{
 		WriteString(Stream, FString::Printf(TEXT("%.9g"), Number));
diff --git a/Engine/Source/Runtime/Json/Public/Policies/JsonPrintPolicy.h b/Engine/Source/Runtime/Json/Public/Policies/JsonPrintPolicy.h
--- a/Engine/Source/Runtime/Json/Public/Policies/JsonPrintPolicy.h
+++ b/Engine/Source/Runtime/Json/Public/Policies/JsonPrintPolicy.h
@@ -39,6 +39,38 @@ struct TJsonPrintPolicy
 		}
 	}
 
+	/**
+	 * Writes a string to the output stream.
+	 *
+	 * @param Stream The stream to write to.
+	 * @param String The string to write.
+	 */
+	static inline void WriteString(FArchive* Stream, FStringView String)
+	{
+		const TCHAR* CharPtr = String.GetData();
+
+		for (int32 CharIndex = 0; CharIndex < String.Len(); ++CharIndex, ++CharPtr)
+		{
+			WriteChar(Stream, *CharPtr);
+		}
+	}
+
+	/**
+	 * Writes a string to the output stream.
+	 *
+	 * @param Stream The stream to write to.
+	 * @param String The string to write.
+	 */
+	static inline void WriteString(FArchive* Stream, FAnsiStringView String)
+	{
+		const ANSICHAR* CharPtr = String.GetData();
+
+		for (int32 CharIndex = 0; CharIndex < String.Len(); ++CharIndex, ++CharPtr)
+		{
+			WriteChar(Stream, *CharPtr);
+		}
+	}
+
 	/**
 	 * Writes a float to the output stream.
 	 *
@@ -74,6 +106,14 @@ inline void TJsonPrintPolicy<TCHAR>::WriteString( FArchive* Stream, const FStrin
 	Stream->Serialize((void*)*String, String.Len() * sizeof(TCHAR));
 }
 
+/**
+ * Specialization for TCHAR that allows direct copying from FStringView data.
+ */
+template <>
+inline void TJsonPrintPolicy<TCHAR>::WriteString(FArchive* Stream, FStringView String)
+{
+	Stream->Serialize((void*)String.GetData(), String.Len() * sizeof(TCHAR));
+}
 
 #if !PLATFORM_TCHAR_IS_CHAR16
 
@@ -89,5 +129,17 @@ inline void TJsonPrintPolicy<UTF16CHAR>::WriteString(FArchive* Stream, const FSt
 	Stream->Serialize((void*)UTF16String.Get(), UTF16String.Length() * sizeof(UTF16CHAR));
 }
 
+/**
+ * Specialization for UTF16CHAR that writes FString data UTF-16.
+ */
+template <>
+inline void TJsonPrintPolicy<UTF16CHAR>::WriteString(FArchive* Stream, FStringView String)
+{
+	// Note: This is a no-op on platforms that are using a 16-bit TCHAR
+	FTCHARToUTF16 UTF16String(String.GetData(), String.Len());
+
+	Stream->Serialize((void*)UTF16String.Get(), UTF16String.Length() * sizeof(UTF16CHAR));
+}
+
 #endif
 
diff --git a/Engine/Source/Runtime/Json/Public/Policies/PrettyJsonPrintPolicy.h b/Engine/Source/Runtime/Json/Public/Policies/PrettyJsonPrintPolicy.h
--- a/Engine/Source/Runtime/Json/Public/Policies/PrettyJsonPrintPolicy.h
+++ b/Engine/Source/Runtime/Json/Public/Policies/PrettyJsonPrintPolicy.h
@@ -16,7 +16,7 @@ struct TPrettyJsonPrintPolicy
 {
 	static inline void WriteLineTerminator( FArchive* Stream )
 	{
-		TJsonPrintPolicy<CharType>::WriteString(Stream, LINE_TERMINATOR);
+		TJsonPrintPolicy<CharType>::WriteString(Stream, FStringView(LINE_TERMINATOR));
 	}
 
 	static inline void WriteTabs( FArchive* Stream, int32 Count )
diff --git a/Engine/Source/Runtime/Json/Public/Serialization/JsonWriter.h b/Engine/Source/Runtime/Json/Public/Serialization/JsonWriter.h
--- a/Engine/Source/Runtime/Json/Public/Serialization/JsonWriter.h
+++ b/Engine/Source/Runtime/Json/Public/Serialization/JsonWriter.h
@@ -107,10 +107,11 @@ public:
 		PreviousTokenWritten = EJsonToken::CurlyOpen;
 	}
 
-	void WriteObjectStart( const FString& Identifier )
+	template<typename IdentifierType>
+	void WriteObjectStart(IdentifierType&& Identifier)
 	{
 		check( Stack.Top() == EJson::Object );
-		WriteIdentifier( Identifier );
+		WriteIdentifier(Forward<IdentifierType>(Identifier));
 
 		PrintPolicy::WriteLineTerminator(Stream);
 		PrintPolicy::WriteTabs(Stream, IndentLevel);
@@ -153,10 +154,11 @@ public:
 		PreviousTokenWritten = EJsonToken::SquareOpen;
 	}
 
-	void WriteArrayStart( const FString& Identifier )
+	template<typename IdentifierType>
+	void WriteArrayStart(IdentifierType&& Identifier)
 	{
 		check( Stack.Top() == EJson::Object );
-		WriteIdentifier( Identifier );
+		WriteIdentifier(Forward<IdentifierType>(Identifier));
 
 		PrintPolicy::WriteSpace( Stream );
 		PrintPolicy::WriteChar(Stream, CharType('['));
@@ -204,6 +206,16 @@ public:
 		PreviousTokenWritten = WriteValueOnly( Value );
 	}
 
+	void WriteValue(FStringView Value)
+	{
+		check(CanWriteValueWithoutIdentifier());
+		WriteCommaIfNeeded();
+
+		PrintPolicy::WriteLineTerminator(Stream);
+		PrintPolicy::WriteTabs(Stream, IndentLevel);
+		PreviousTokenWritten = WriteValueOnly(Value);
+	}
+
 	void WriteValue(const FString& Value)
 	{
 		check(CanWriteValueWithoutIdentifier());
@@ -214,20 +226,20 @@ public:
 		PreviousTokenWritten = WriteValueOnly(Value);
 	}
 
-	template <class FValue>
-	void WriteValue(const FString& Identifier, FValue Value)
+	template<class FValue, typename IdentifierType>
+	void WriteValue(IdentifierType&& Identifier, FValue Value)
 	{
 		check( Stack.Top() == EJson::Object );
-		WriteIdentifier( Identifier );
+		WriteIdentifier(Forward<IdentifierType>(Identifier));
 
 		PrintPolicy::WriteSpace(Stream);
 		PreviousTokenWritten = WriteValueOnly(MoveTemp(Value));
 	}
 
-	template<class ElementType>
-	void WriteValue(const FString& Identifier, const TArray<ElementType>& Array)
+	template<class ElementType, typename IdentifierType>
+	void WriteValue(IdentifierType&& Identifier, const TArray<ElementType>& Array)
 	{
-		WriteArrayStart(Identifier);
+		WriteArrayStart(Forward<IdentifierType>(Identifier));
 		for (int Idx = 0; Idx < Array.Num(); Idx++)
 		{
 			WriteValue(Array[Idx]);
@@ -235,9 +247,9 @@ public:
 		WriteArrayEnd();
 	}
 
-	void WriteValue(const FString& Identifier, const TCHAR* Value)
+	void WriteValue(FStringView Identifier, const TCHAR* Value)
 	{
-		WriteValue(Identifier, FString(Value));
+		WriteValue(Identifier, FStringView(Value));
 	}
 
 	// WARNING: THIS IS DANGEROUS. Use this only if you know for a fact that the Value is valid JSON!
@@ -252,19 +264,20 @@ public:
 		PreviousTokenWritten = EJsonToken::String;
 	}
 
-	void WriteNull( const FString& Identifier )
+	template<typename IdentifierType>
+	void WriteNull(IdentifierType&& Identifier)
 	{
-		WriteValue(Identifier, nullptr);
+		WriteValue(Forward<IdentifierType>(Identifier), nullptr);
 	}
 
-	void WriteValue( const TCHAR* Value )
+	void WriteValue(const TCHAR* Value)
 	{
-		WriteValue(FString(Value));
+		WriteValue(FStringView(Value));
 	}
 
 	// WARNING: THIS IS DANGEROUS. Use this only if you know for a fact that the Value is valid JSON!
 	// Use this to insert the results of a different JSON Writer in.
-	void WriteRawJSONValue( const FString& Value )
+	void WriteRawJSONValue(const FString& Value)
 	{
 		check(CanWriteValueWithoutIdentifier());
 		WriteCommaIfNeeded();
@@ -299,10 +312,11 @@ public:
 	/**
 	 * WriteValue("Foo", Bar) should be equivalent to WriteIdentifierPrefix("Foo"), WriteValue(Bar)
 	 */
-	void WriteIdentifierPrefix(const FString& Identifier)
+	template<typename IdentifierType>
+	void WriteIdentifierPrefix(IdentifierType&& Identifier)
 	{
 		check(Stack.Top() == EJson::Object);
-		WriteIdentifier(Identifier);
+		WriteIdentifier(Forward<IdentifierType>(Identifier));
 		PrintPolicy::WriteSpace(Stream);
 		PreviousTokenWritten = EJsonToken::Identifier;
 	}
@@ -342,19 +356,54 @@ protected:
 		}
 	}
 
-	FORCEINLINE void WriteIdentifier( const FString& Identifier )
+	void WriteIdentifier(const ANSICHAR* Identifier)
+	{
+		WriteCommaIfNeeded();
+		PrintPolicy::WriteLineTerminator(Stream);
+
+		PrintPolicy::WriteTabs(Stream, IndentLevel);
+		WriteStringValue(FAnsiStringView(Identifier));
+		PrintPolicy::WriteChar(Stream, CharType(':'));
+	}
+	
+	void WriteIdentifier(const TCHAR* Identifier)
+	{
+		WriteCommaIfNeeded();
+		PrintPolicy::WriteLineTerminator(Stream);
+
+		PrintPolicy::WriteTabs(Stream, IndentLevel);
+		WriteStringValue(FStringView(Identifier));
+		PrintPolicy::WriteChar(Stream, CharType(':'));
+	}
+
+	void WriteIdentifier(FStringView Identifier)
+	{
+		WriteCommaIfNeeded();
+		PrintPolicy::WriteLineTerminator(Stream);
+
+		PrintPolicy::WriteTabs(Stream, IndentLevel);
+		WriteStringValue(Identifier);
+		PrintPolicy::WriteChar(Stream, CharType(':'));
+	}
+
+	void WriteIdentifier(const FText& Identifier)
+	{
+		WriteIdentifier(Identifier.ToString()); // Does not copy
+	}
+
+	FORCEINLINE void WriteIdentifier(const FString& Identifier)
 	{
 		WriteCommaIfNeeded();
 		PrintPolicy::WriteLineTerminator(Stream);
 
 		PrintPolicy::WriteTabs(Stream, IndentLevel);
-		WriteStringValue( Identifier );
+		WriteStringValue(FStringView(Identifier));
 		PrintPolicy::WriteChar(Stream, CharType(':'));
 	}
 
 	FORCEINLINE EJsonToken WriteValueOnly(bool Value)
 	{
-		PrintPolicy::WriteString(Stream, Value ? TEXT("true") : TEXT("false"));
+		PrintPolicy::WriteString(Stream, Value ? TEXTVIEW("true") : TEXTVIEW("false"));
 		return Value ? EJsonToken::True : EJsonToken::False;
 	}
 
@@ -379,26 +428,127 @@ protected:
 
 	FORCEINLINE EJsonToken WriteValueOnly(int64 Value)
 	{
-		PrintPolicy::WriteString(Stream, FString::Printf(TEXT("%lld"), Value));
+		PrintPolicy::WriteString(Stream, WriteToString<32>(Value));
+		return EJsonToken::Number;
+	}
+
+	EJsonToken WriteValueOnly(uint64 Value)
+	{
+		PrintPolicy::WriteString(Stream, WriteToString<32>(Value));
 		return EJsonToken::Number;
 	}
 
 	FORCEINLINE EJsonToken WriteValueOnly(TYPE_OF_NULLPTR)
 	{
-		PrintPolicy::WriteString(Stream, TEXT("null"));
+		PrintPolicy::WriteString(Stream, TEXTVIEW("null"));
 		return EJsonToken::Null;
 	}
 
-	FORCEINLINE EJsonToken WriteValueOnly(const FString& Value)
+	FORCEINLINE EJsonToken WriteValueOnly(const TCHAR* Value)
+	{
+		WriteStringValue(FStringView(Value));
+		return EJsonToken::String;
+	}
+
+	FORCEINLINE EJsonToken WriteValueOnly(FStringView Value)
 	{
 		WriteStringValue(Value);
 		return EJsonToken::String;
 	}
 
-	virtual void WriteStringValue( const FString& String )
+	virtual void WriteStringValue(FAnsiStringView String)
+	{
+		PrintPolicy::WriteChar(Stream, CharType('"'));
+		WriteEscapedString(String);
+		PrintPolicy::WriteChar(Stream, CharType('"'));
+	}
+
+	virtual void WriteStringValue(FStringView String)
+	{
+		PrintPolicy::WriteChar(Stream, CharType('"'));
+		WriteEscapedString(String);
+		PrintPolicy::WriteChar(Stream, CharType('"'));
+	}
+
+	virtual void WriteStringValue(const FString& String)
+	{
+		WriteStringValue(FStringView(String));
+	}
+
+	template<typename InCharType>
+	void WriteEscapedString(TStringView<InCharType> InView)
 	{
-		FString OutString = EscapeJsonString(String);
-		PrintPolicy::WriteString(Stream, OutString);
+		auto NeedsEscaping = [](InCharType Char) -> bool
+		{
+			switch (Char)
+			{
+			case TCHAR('\\'): return true;
+			case TCHAR('\n'): return true;
+			case TCHAR('\t'): return true;
+			case TCHAR('\b'): return true;
+			case TCHAR('\f'): return true;
+			case TCHAR('\r'): return true;
+			case TCHAR('\"'): return true;
+			default:
+				// Must escape control characters
+				if (Char >= TCHAR(32))
+				{
+					return false;
+				}
+				else
+				{
+					return true;
+				}
+			}
+		};
+
+		// Write successive runs of unescaped and escaped characters until the view is exhausted
+		while (!InView.IsEmpty())
+		{
+			 // In case we are handed a very large string, avoid checking all of it at once without writing anything
+			constexpr int32 LongestRun = 2048;
+			int32 EndIndex = 0;
+			for (; EndIndex < InView.Len() && EndIndex < LongestRun; ++EndIndex)
+			{
+				if (NeedsEscaping(InView[EndIndex]))
+				{ 
+					break;
+				}
+			}
+			if (TStringView<InCharType> Blittable = InView.Left(EndIndex); !Blittable.IsEmpty())
+			{
+				PrintPolicy::WriteString(Stream, Blittable);
+			}
+			InView.RightChopInline(EndIndex);
+			for (EndIndex = 0; EndIndex < InView.Len(); ++EndIndex)
+			{
+				TCHAR Char = InView[EndIndex]; 
+				switch (Char)
+				{
+				case TCHAR('\\'): PrintPolicy::WriteString(Stream, TEXTVIEW("\\\\")); continue;
+				case TCHAR('\n'): PrintPolicy::WriteString(Stream, TEXTVIEW("\\n")); continue;
+				case TCHAR('\t'): PrintPolicy::WriteString(Stream, TEXTVIEW("\\t")); continue;
+				case TCHAR('\b'): PrintPolicy::WriteString(Stream, TEXTVIEW("\\b")); continue;
+				case TCHAR('\f'): PrintPolicy::WriteString(Stream, TEXTVIEW("\\f")); continue;
+				case TCHAR('\r'): PrintPolicy::WriteString(Stream, TEXTVIEW("\\r")); continue;
+				case TCHAR('\"'): PrintPolicy::WriteString(Stream, TEXTVIEW("\\\"")); continue;
+				default: break;
+				}
+
+				// Must escape control characters
+				if (Char >= TCHAR(32))
+				{
+					break;
+				}
+				else
+				{
+					TAnsiStringBuilder<8> Builder;
+					Builder.Appendf("\\u%04x", Char);
+					PrintPolicy::WriteString(Stream, Builder.ToView());
+				}
+			}
+			InView.RightChopInline(EndIndex);
+		}
 	}
 
 	FArchive* const Stream;
