Fix JSON Reader UTF8 parsing.
Add ability to use JSON Reader with a string view.
Add explicit UTF8 smoke tests.
Change uses of TJsonReader<UCS2CHAR> to TJsonReader<WIDECHAR> since TStringBuilder<UCS2CHAR> doesn't compile.
Change uses of TJsonReader<char> to TJsonReader<UTF8CHAR>.  There isn't much use for TJsonReader<char> as ASCII is a subset of UTF8 so TJsonReader<UTF8CHAR> is always preferable.

diff --git a/Engine/Source/Editor/UnrealEd/Private/Commandlets/TextAssetCommandlet.cpp b/Engine/Source/Editor/UnrealEd/Private/Commandlets/TextAssetCommandlet.cpp
--- a/Engine/Source/Editor/UnrealEd/Private/Commandlets/TextAssetCommandlet.cpp
+++ b/Engine/Source/Editor/UnrealEd/Private/Commandlets/TextAssetCommandlet.cpp
@@ -783,7 +783,7 @@ bool UTextAssetCommandlet::DoTextAssetProcessing(const FProcessingArgs& InArgs)
 						TRACE_CPUPROFILER_EVENT_SCOPE(UTextAssetCommandlet::VerifyJson);
 						FArchive* File = IFileManager::Get().CreateFileReader(*DestinationFilename);
 						TSharedPtr< FJsonObject > RootObject;
-						TSharedRef< TJsonReader<char> > Reader = TJsonReaderFactory<char>::Create(File);
+						TSharedRef< TJsonReader<UTF8CHAR> > Reader = TJsonReaderFactory<UTF8CHAR>::Create(File);
 						ensure(FJsonSerializer::Deserialize(Reader, RootObject));
 						delete File;
 					}
diff --git a/Engine/Source/Runtime/Core/Public/Misc/KeyChainUtilities.h b/Engine/Source/Runtime/Core/Public/Misc/KeyChainUtilities.h
--- a/Engine/Source/Runtime/Core/Public/Misc/KeyChainUtilities.h
+++ b/Engine/Source/Runtime/Core/Public/Misc/KeyChainUtilities.h
@@ -152,7 +152,7 @@ namespace KeyChainUtilities
 		FArchive* File = IFileManager::Get().CreateFileReader(*InFilename);
 		checkf(File != nullptr, TEXT("Specified crypto keys cache '%s' does not exist!"), *InFilename);
 		TSharedPtr<FJsonObject> RootObject;
-		TSharedRef<TJsonReader<char>> Reader = TJsonReaderFactory<char>::Create(File);
+		TSharedRef<TJsonReader<UTF8CHAR>> Reader = TJsonReaderFactory<UTF8CHAR>::Create(File);
 		if (FJsonSerializer::Deserialize(Reader, RootObject))
 		{
 			const TSharedPtr<FJsonObject>* EncryptionKeyObject;
diff --git a/Engine/Source/Runtime/CoreUObject/Private/Serialization/Formatters/JsonArchiveInputFormatter.cpp b/Engine/Source/Runtime/CoreUObject/Private/Serialization/Formatters/JsonArchiveInputFormatter.cpp
--- a/Engine/Source/Runtime/CoreUObject/Private/Serialization/Formatters/JsonArchiveInputFormatter.cpp
+++ b/Engine/Source/Runtime/CoreUObject/Private/Serialization/Formatters/JsonArchiveInputFormatter.cpp
@@ -22,7 +22,7 @@ FJsonArchiveInputFormatter::FJsonArchiveInputFormatter(FArchive& InInner, TFunct
 	Inner.ArAllowLazyLoading = false;
 
 	TSharedPtr< FJsonObject > RootObject;
-	TSharedRef< TJsonReader<char> > Reader = TJsonReaderFactory<char>::Create(&InInner);
+	TSharedRef< TJsonReader<UTF8CHAR> > Reader = TJsonReaderFactory<UTF8CHAR>::Create(&InInner);
 	ensure(FJsonSerializer::Deserialize(Reader, RootObject, FJsonSerializer::EFlags::StoreNumbersAsStrings));
 	ValueStack.Add(MakeShared<FJsonValueObject>(RootObject));
 
diff --git a/Engine/Source/Runtime/Json/Private/Tests/JsonTests.cpp b/Engine/Source/Runtime/Json/Private/Tests/JsonTests.cpp
--- a/Engine/Source/Runtime/Json/Private/Tests/JsonTests.cpp
+++ b/Engine/Source/Runtime/Json/Private/Tests/JsonTests.cpp
@@ -325,6 +325,66 @@ bool FJsonAutomationTest::RunTest(const FString& Parameters)
 		check(OutputString == TestOutput);
 	}
 
+	// String Test UTF8
+	{
+		// UTF8TEXT will prepend the first u8 literal specifier
+		// UTF8TEXT does a cast so we can't add it each line and still get the literals to concatenate
+		const UTF8CHAR* InputString = UTF8TEXT(
+			"{"
+				u8"\"Value\":\"Some String, Escape Chars: \\\\, \\\", \\/, \\b, \\f, \\n, \\r, \\t, \\u002B\\uD83D\\uDE10\","
+				u8"\"Value1\":\"Greek String, Σὲ γνωρίζω ἀπὸ τὴν κόψη\","
+				u8"\"Value2\":\"Thai String, สิบสองกษัตริย์ก่อนหน้าแลถัดไป\","
+				u8"\"Value3\":\"Hello world, Καλημέρα κόσμε, コンニチハ\""
+			u8"}"
+		);
+
+		TSharedRef< TJsonReader<UTF8CHAR> > Reader = TJsonReaderFactory<UTF8CHAR>::CreateFromView( InputString );
+
+		TSharedPtr<FJsonObject> Object;
+		bool bSuccessful = FJsonSerializer::Deserialize(Reader, Object);
+		check(bSuccessful);
+		check( Object.IsValid() );
+
+		{
+			const TSharedPtr<FJsonValue>* Value = Object->Values.Find(TEXT("Value"));
+			check(Value && (*Value)->Type == EJson::String);
+			const FString String = (*Value)->AsString();
+			check(String == TEXT("Some String, Escape Chars: \\, \", /, \b, \f, \n, \r, \t, +😐"));
+		}
+		{
+			const TSharedPtr<FJsonValue>* Value = Object->Values.Find(TEXT("Value1"));
+			check(Value && (*Value)->Type == EJson::String);
+			const FString String = (*Value)->AsString();
+			check(String == TEXT("Greek String, Σὲ γνωρίζω ἀπὸ τὴν κόψη"));
+		}
+		{
+			const TSharedPtr<FJsonValue>* Value = Object->Values.Find(TEXT("Value2"));
+			check(Value && (*Value)->Type == EJson::String);
+			const FString String = (*Value)->AsString();
+			check(String == TEXT("Thai String, สิบสองกษัตริย์ก่อนหน้าแลถัดไป"));
+		}
+		{
+			const TSharedPtr<FJsonValue>* Value = Object->Values.Find(TEXT("Value3"));
+			check(Value && (*Value)->Type == EJson::String);
+			const FString String = (*Value)->AsString();
+			check(String == TEXT("Hello world, Καλημέρα κόσμε, コンニチハ"));
+		}
+
+		FString OutputString;
+		TSharedRef< FCondensedJsonStringWriter > Writer = FCondensedJsonStringWriterFactory::Create( &OutputString );
+		verify( FJsonSerializer::Serialize( Object.ToSharedRef(), Writer ) );
+
+		// Note: The literal prefix for the string (u8, L), must be present for every contatenated string, not just the first one
+		const FString TestOutput =
+			TEXT("{")
+				TEXT("\"Value\":\"Some String, Escape Chars: \\\\, \\\", /, \\b, \\f, \\n, \\r, \\t, +😐\",")
+				TEXT("\"Value1\":\"Greek String, Σὲ γνωρίζω ἀπὸ τὴν κόψη\",")
+				TEXT("\"Value2\":\"Thai String, สิบสองกษัตริย์ก่อนหน้าแลถัดไป\",")
+				TEXT("\"Value3\":\"Hello world, Καλημέρα κόσμε, コンニチハ\"")
+			TEXT("}");
+		check(OutputString == TestOutput);
+	}
+
 	// Number Test
 	{
 		const FString InputString =
diff --git a/Engine/Source/Runtime/Json/Public/Serialization/JsonReader.h b/Engine/Source/Runtime/Json/Public/Serialization/JsonReader.h
--- a/Engine/Source/Runtime/Json/Public/Serialization/JsonReader.h
+++ b/Engine/Source/Runtime/Json/Public/Serialization/JsonReader.h
@@ -5,6 +5,7 @@
 #include "CoreMinimal.h"
 #include "Serialization/JsonTypes.h"
 #include "Serialization/BufferReader.h"
+#include "Misc/StringBuilder.h"
 
 class Error;
 
@@ -212,7 +213,7 @@ protected:
 	 *
 	 * @param InStream An archive containing the input.
 	 */
-	TJsonReader(FArchive* InStream)
+	explicit TJsonReader(FArchive* InStream)
 		: ParseState()
 		, CurrentToken(EJsonToken::None)
 		, Stream(InStream)
@@ -553,7 +554,19 @@ private:
 
 	bool ParseStringToken()
 	{
-		FString String;
+		TStringBuilderWithBuffer<CharType, 512> StringBuffer;
+		TStringBuilderWithBuffer<WIDECHAR, 512> UTF16CodePoints;
+
+		// Add escaped surrogate pairs
+		auto ConditionallyAddCodePoints = [&StringBuffer, &UTF16CodePoints]() -> void
+		{
+			if (UTF16CodePoints.Len() > 0)
+			{
+				// Will convert to CharType encoding if needed
+				StringBuffer.Append(UTF16CodePoints);
+				UTF16CodePoints.Reset();
+			}
+		};
 
 		while (true)
 		{
@@ -572,6 +585,7 @@ private:
 
 			if (Char == CharType('\"'))
 			{
+				ConditionallyAddCodePoints();
 				break;
 			}
 
@@ -583,14 +597,19 @@ private:
 				}
 				++CharacterNumber;
 
+				if (Char != CharType('u'))
+				{
+					ConditionallyAddCodePoints();
+				}
+
 				switch (Char)
 				{
-				case CharType('\"'): case CharType('\\'): case CharType('/'): String += Char; break;
-				case CharType('f'): String += CharType('\f'); break;
-				case CharType('r'): String += CharType('\r'); break;
-				case CharType('n'): String += CharType('\n'); break;
-				case CharType('b'): String += CharType('\b'); break;
-				case CharType('t'): String += CharType('\t'); break;
+				case CharType('\"'): case CharType('\\'): case CharType('/'): StringBuffer.AppendChar(Char); break;
+				case CharType('f'): StringBuffer.AppendChar(CharType('\f')); break;
+				case CharType('r'): StringBuffer.AppendChar(CharType('\r')); break;
+				case CharType('n'): StringBuffer.AppendChar(CharType('\n')); break;
+				case CharType('b'): StringBuffer.AppendChar(CharType('\b')); break;
+				case CharType('t'): StringBuffer.AppendChar(CharType('\t')); break;
 				case CharType('u'):
 					// 4 hex digits, like \uAB23, which is a 16 bit number that we would usually see as 0xAB23
 					{
@@ -622,7 +641,7 @@ private:
 							HexNum += HexDigit * (int32)FMath::Pow(16.f, (float)Radix);
 						}
 
-						String += (FString::ElementType)HexNum;
+						UTF16CodePoints.AppendChar((UTF16CHAR)HexNum);
 					}
 					break;
 
@@ -633,11 +652,12 @@ private:
 			}
 			else
 			{
-				String += Char;
+				ConditionallyAddCodePoints();
+				StringBuffer.AppendChar(Char);
 			}
 		}
 
-		StringValue = MoveTemp(String);
+		StringValue = FString(StringBuffer);
 
 		// Inline combine any surrogate pairs in the data when loading into a UTF-32 string
 		StringConv::InlineCombineSurrogates(StringValue);
@@ -894,7 +914,7 @@ protected:
 	 *
 	 * @param JsonString The Json string to parse.
 	 */
-	FJsonStringReader(const FString& JsonString)
+	explicit FJsonStringReader(const FString& JsonString)
 		: Content(JsonString)
 		, Reader(nullptr)
 	{
@@ -906,7 +926,7 @@ protected:
 	 *
 	 * @param JsonString The Json string to parse.
 	 */
-	FJsonStringReader(FString&& JsonString)
+	explicit FJsonStringReader(FString&& JsonString)
 		: Content(MoveTemp(JsonString))
 		, Reader(nullptr)
 	{
@@ -921,8 +941,6 @@ protected:
 		}
 
 		Reader = MakeUnique<FBufferReader>((void*)*Content, Content.Len() * sizeof(TCHAR), false);
-		check(Reader.IsValid());
-
 		Stream = Reader.Get();
 	}
 
@@ -931,6 +949,50 @@ protected:
 	TUniquePtr<FBufferReader> Reader;
 };
 
+template <class CharType>
+class TJsonStringViewReader
+	: public TJsonReader<CharType>
+{
+public:
+
+	static TSharedRef<TJsonStringViewReader> Create(TStringView<CharType> JsonString)
+	{
+		return MakeShareable(new TJsonStringViewReader(JsonString));
+	}
+
+public:
+
+	virtual ~TJsonStringViewReader() = default;
+
+protected:
+
+	/**
+	 * Parses a string containing Json information.
+	 *
+	 * @param JsonString The Json string to parse.
+	 */
+	explicit TJsonStringViewReader(TStringView<CharType> JsonString)
+		: Content(JsonString)
+		, Reader(nullptr)
+	{
+		InitReader();
+	}
+
+	void InitReader()
+	{
+		if (Content.IsEmpty())
+		{
+			return;
+		}
+
+		Reader = MakeUnique<FBufferReader>((void*)Content.GetData(), Content.Len() * sizeof(CharType), false);
+		TJsonReader<CharType>::Stream = Reader.Get();
+	}
+
+protected:
+	TStringView<CharType> Content;
+	TUniquePtr<FBufferReader> Reader;
+};
 
 template <class CharType = TCHAR>
 class TJsonReaderFactory
@@ -951,4 +1013,9 @@ public:
 	{
 		return TJsonReader<CharType>::Create(Stream);
 	}
+
+	static TSharedRef<TJsonReader<CharType>> CreateFromView(TStringView<CharType> JsonString)
+	{
+		return TJsonStringViewReader<CharType>::Create(JsonString);
+	}
 };
diff --git a/Engine/Source/Runtime/Serialization/Public/Backends/JsonStructDeserializerBackend.h b/Engine/Source/Runtime/Serialization/Public/Backends/JsonStructDeserializerBackend.h
--- a/Engine/Source/Runtime/Serialization/Public/Backends/JsonStructDeserializerBackend.h
+++ b/Engine/Source/Runtime/Serialization/Public/Backends/JsonStructDeserializerBackend.h
@@ -32,7 +32,7 @@ public:
 	 * @param Archive The archive to deserialize from.
 	 */
 	FJsonStructDeserializerBackend( FArchive& Archive )
-		: JsonReader(TJsonReader<UCS2CHAR>::Create(&Archive))
+		: JsonReader(TJsonReader<WIDECHAR>::Create(&Archive))
 	{ }
 
 public:
@@ -58,7 +58,7 @@ protected:
 		return LastNotation;
 	}
 
-	TSharedRef<TJsonReader<UCS2CHAR>>& GetReader()
+	TSharedRef<TJsonReader<WIDECHAR>>& GetReader()
 	{
 		return JsonReader;
 	}
@@ -72,5 +72,5 @@ private:
 	EJsonNotation LastNotation;
 
 	/** Holds the Json reader used for the actual reading of the archive. */
-	TSharedRef<TJsonReader<UCS2CHAR>> JsonReader;
+	TSharedRef<TJsonReader<WIDECHAR>> JsonReader;
 };
