Json: Added a log message with details when we fail to parse a Json.

I've had an interesting case relating to it. A licensee had Achievements.json fille that would fail to parse even though at first sight it seemed to be complate. It turned out that it was caused by members like the following:

"Field" : 0001

Such leading zeros for Number elements are not allowed by the Json specification and our parser confirms to it. The additional log will output messages like:

Poorly formed Json Number Token. Line: 4 Ch: 17

which should make such issues much more clear.

#rnx

diff --git a/Engine/Source/Runtime/Json/Public/Serialization/JsonSerializerMacros.h b/Engine/Source/Runtime/Json/Public/Serialization/JsonSerializerMacros.h
--- a/Engine/Source/Runtime/Json/Public/Serialization/JsonSerializerMacros.h
+++ b/Engine/Source/Runtime/Json/Public/Serialization/JsonSerializerMacros.h
@@ -1056,6 +1056,7 @@ struct FJsonSerializable
 			Serialize(Serializer, false);
 			return true;
 		}
+		UE_LOG(LogJson, Warning, TEXT("Failed to parse Json from a string: %s"), *JsonReader->GetErrorMessage());
 		return false;
 	}
 
