Fixed TJsonPrintPolicy::WriteString() doing per-character casts instead of a proper string conversion.

diff --git a/Engine/Source/Runtime/Json/Public/Policies/JsonPrintPolicy.h b/Engine/Source/Runtime/Json/Public/Policies/JsonPrintPolicy.h
--- a/Engine/Source/Runtime/Json/Public/Policies/JsonPrintPolicy.h
+++ b/Engine/Source/Runtime/Json/Public/Policies/JsonPrintPolicy.h
@@ -31,12 +31,8 @@ struct TJsonPrintPolicy
 	 */
 	static inline void WriteString( FArchive* Stream, const FString& String )
 	{
-		const TCHAR* CharPtr = *String;
-
-		for (int32 CharIndex = 0; CharIndex < String.Len(); ++CharIndex, ++CharPtr)
-		{
-			WriteChar(Stream, static_cast<CharType>(*CharPtr));
-		}
+		auto Conv = StringCast<CharType>(*String, String.Len());
+		Stream->Serialize((void*)Conv.Get(), Conv.Length() * sizeof(CharType));
 	}
 
 	/**
@@ -47,12 +43,8 @@ struct TJsonPrintPolicy
 	 */
 	static inline void WriteString(FArchive* Stream, FStringView String)
 	{
-		const TCHAR* CharPtr = String.GetData();
-
-		for (int32 CharIndex = 0; CharIndex < String.Len(); ++CharIndex, ++CharPtr)
-		{
-			WriteChar(Stream, *CharPtr);
-		}
+		auto Conv = StringCast<CharType>(String.GetData(), String.Len());
+		Stream->Serialize((void*)Conv.Get(), Conv.Length() * sizeof(CharType));
 	}
 
 	/**
@@ -63,12 +55,8 @@ struct TJsonPrintPolicy
 	 */
 	static inline void WriteString(FArchive* Stream, FAnsiStringView String)
 	{
-		const ANSICHAR* CharPtr = String.GetData();
-
-		for (int32 CharIndex = 0; CharIndex < String.Len(); ++CharIndex, ++CharPtr)
-		{
-			WriteChar(Stream, *CharPtr);
-		}
+		auto Conv = StringCast<CharType>(String.GetData(), String.Len());
+		Stream->Serialize((void*)Conv.Get(), Conv.Length() * sizeof(CharType));
 	}
 
 	/**
@@ -106,14 +94,6 @@ inline void TJsonPrintPolicy<TCHAR>::WriteString( FArchive* Stream, const FStrin
 	Stream->Serialize((void*)*String, String.Len() * sizeof(TCHAR));
 }
 
-/**
- * Specialization for TCHAR that allows direct copying from FStringView data.
- */
-template <>
-inline void TJsonPrintPolicy<TCHAR>::WriteString(FArchive* Stream, FStringView String)
-{
-	Stream->Serialize((void*)String.GetData(), String.Len() * sizeof(TCHAR));
-}
 
 #if !PLATFORM_TCHAR_IS_CHAR16
 
