Fixed JsonSerializer output where NumberStrings were being written as numbers, not strings, causing data loss.

Reported from: https://udn.unrealengine.com/s/question/0D54z00008fjNFQCA2/fjsonserializerserialize-loses-precision-for-fjsonvaluenumberstring

diff --git a/Engine/Source/Runtime/Json/Public/Dom/JsonValue.h b/Engine/Source/Runtime/Json/Public/Dom/JsonValue.h
--- a/Engine/Source/Runtime/Json/Public/Dom/JsonValue.h
+++ b/Engine/Source/Runtime/Json/Public/Dom/JsonValue.h
@@ -83,6 +83,9 @@ public:
 	/** Tries to convert this value to an object, returning false if not possible */
 	virtual bool TryGetObject(TSharedPtr<FJsonObject>*& Object) { return false; }
 
+	/** Returns whether or not a caller should prefer a string representation of the value, rather than the natural JSON type */
+	virtual bool PreferStringRepresentation() const { return false; }
+
 	/** Returns true if this value is a 'null' */
 	bool IsNull() const { return Type == EJson::Null || Type == EJson::None; }
 
@@ -181,6 +184,7 @@ public:
 	virtual bool TryGetNumber(uint32& OutValue) const override { return LexTryParseString(OutValue, *Value); }
 	virtual bool TryGetNumber(uint64& OutValue) const override { return LexTryParseString(OutValue, *Value); }
 	virtual bool TryGetBool(bool& OutBool) const override { OutBool = Value.ToBool(); return true; }
+	virtual bool PreferStringRepresentation() const override { return true; }
 
 protected:
 	FString Value;
diff --git a/Engine/Source/Runtime/Json/Public/Serialization/JsonSerializer.h b/Engine/Source/Runtime/Json/Public/Serialization/JsonSerializer.h
--- a/Engine/Source/Runtime/Json/Public/Serialization/JsonSerializer.h
+++ b/Engine/Source/Runtime/Json/Public/Serialization/JsonSerializer.h
@@ -321,11 +321,25 @@ private:
 				{
 					if (Element->Identifier.IsEmpty())
 					{
-						Writer.WriteValue(Element->Value->AsNumber());
+						if (Element->Value->PreferStringRepresentation())
+						{
+							Writer.WriteValue(Element->Value->AsString());
+						}
+						else
+						{
+							Writer.WriteValue(Element->Value->AsNumber());
+						}
 					}
 					else
 					{
-						Writer.WriteValue(Element->Identifier, Element->Value->AsNumber());
+						if (Element->Value->PreferStringRepresentation())
+						{
+							Writer.WriteValue(Element->Identifier, Element->Value->AsString());
+						}
+						else
+						{
+							Writer.WriteValue(Element->Identifier, Element->Value->AsNumber());
+						}
 					}
 				}
 				break;
