Fixed numberstrings being written out with quotes.

diff --git a/Engine/Source/Runtime/Json/Public/Serialization/JsonSerializer.h b/Engine/Source/Runtime/Json/Public/Serialization/JsonSerializer.h
--- a/Engine/Source/Runtime/Json/Public/Serialization/JsonSerializer.h
+++ b/Engine/Source/Runtime/Json/Public/Serialization/JsonSerializer.h
@@ -323,7 +323,7 @@ private:
 					{
 						if (Element->Value->PreferStringRepresentation())
 						{
-							Writer.WriteValue(Element->Value->AsString());
+							Writer.WriteRawJSONValue(Element->Value->AsString());
 						}
 						else
 						{
@@ -334,7 +334,7 @@ private:
 					{
 						if (Element->Value->PreferStringRepresentation())
 						{
-							Writer.WriteValue(Element->Identifier, Element->Value->AsString());
+							Writer.WriteRawJSONValue(Element->Identifier, Element->Value->AsString());
 						}
 						else
 						{
