Deleted redundant TJsonPrintPolicy<TCHAR>::WriteString() specialization which should have been deleted during the merge of CL# 23408835.

diff --git a/Engine/Source/Runtime/Json/Public/Policies/JsonPrintPolicy.h b/Engine/Source/Runtime/Json/Public/Policies/JsonPrintPolicy.h
--- a/Engine/Source/Runtime/Json/Public/Policies/JsonPrintPolicy.h
+++ b/Engine/Source/Runtime/Json/Public/Policies/JsonPrintPolicy.h
@@ -85,16 +85,6 @@ struct TJsonPrintPolicy
 };
 
 
-/**
- * Specialization for TCHAR that allows direct copying from FString data.
- */
-template <>
-inline void TJsonPrintPolicy<TCHAR>::WriteString( FArchive* Stream, const FString& String )
-{
-	Stream->Serialize((void*)*String, String.Len() * sizeof(TCHAR));
-}
-
-
 #if !PLATFORM_TCHAR_IS_CHAR16
 
 /**
