JsonSerializerMacros: adding float array parsing, addd WITHDEFAULT macros to avoid the need for TOptional values to be able to fill blueprint accessible structs.

diff --git a/Engine/Source/Runtime/Json/Public/Serialization/JsonSerializerMacros.h b/Engine/Source/Runtime/Json/Public/Serialization/JsonSerializerMacros.h
--- a/Engine/Source/Runtime/Json/Public/Serialization/JsonSerializerMacros.h
+++ b/Engine/Source/Runtime/Json/Public/Serialization/JsonSerializerMacros.h
@@ -24,6 +24,16 @@
 #define JSON_SERIALIZE(JsonName, JsonValue) \
 		Serializer.Serialize(TEXT(JsonName), JsonValue)
 
+#define JSON_SERIALIZE_WITHDEFAULT(JsonName, JsonValue, DefaultJsonValue) \
+		if (Serializer.IsLoading()) \
+		{ \
+			if (!Serializer.GetObject()->HasField(TEXT(JsonName))) \
+           	{ \
+           		JsonValue = DefaultJsonValue; \
+           	} \
+        } \
+		Serializer.Serialize(TEXT(JsonName), JsonValue);
+		
 #define JSON_SERIALIZE_OPTIONAL(JsonName, OptionalJsonValue) \
 		if (Serializer.IsLoading()) \
 		{ \
@@ -42,6 +52,26 @@
 
 #define JSON_SERIALIZE_ARRAY(JsonName, JsonArray) \
 		Serializer.SerializeArray(TEXT(JsonName), JsonArray)
+		
+#define JSON_SERIALIZE_ARRAY_WITHDEFAULT(JsonName, JsonArray, DefaultArray) \
+		if (Serializer.IsLoading()) \
+		{ \
+			if (!Serializer.GetObject()->HasField(TEXT(JsonName))) \
+			{ \
+				JsonArray = DefaultArray; \
+			} \
+		} \
+		Serializer.SerializeArray(TEXT(JsonName), JsonArray);
+		
+#define JSON_SERIALIZE_OBJECT_WITHDEFAULT(JsonName, JsonObject, DefaultObject) \
+		if (Serializer.IsLoading()) \
+		{ \
+			if (!Serializer.GetObject()->HasField(TEXT(JsonName))) \
+			{ \
+				JsonObject = DefaultObject; \
+			} \
+		} \
+		Serializer.SerializeArray(TEXT(JsonName), JsonArray);
 
 #define JSON_SERIALIZE_MAP(JsonName, JsonMap) \
 		Serializer.SerializeMap(TEXT(JsonName), JsonMap)
@@ -102,6 +132,32 @@
 			} \
 			Serializer.EndArray(); \
 		}
+		
+#define JSON_SERIALIZE_ARRAY_SERIALIZABLE_WITHDEFAULT(JsonName, JsonArray, ElementType, DefaultArray) \
+		if (Serializer.IsLoading()) \
+		{ \
+			if (Serializer.GetObject()->HasTypedField<EJson::Array>(JsonName)) \
+			{ \
+				for (auto It = Serializer.GetObject()->GetArrayField(JsonName).CreateConstIterator(); It; ++It) \
+				{ \
+					ElementType* Obj = new(JsonArray) ElementType(); \
+					Obj->FromJson((*It)->AsObject()); \
+				} \
+			} \
+			else \
+			{ \
+				JsonArray = DefaultArray; \
+			} \
+		} \
+		else \
+		{ \
+			Serializer.StartArray(JsonName); \
+			for (auto It = JsonArray.CreateIterator(); It; ++It) \
+			{ \
+				It->Serialize(Serializer, false); \
+			} \
+			Serializer.EndArray(); \
+		}
 
 #define JSON_SERIALIZE_OPTIONAL_ARRAY_SERIALIZABLE(JsonName, OptionalJsonArray, ElementType) \
 		if (Serializer.IsLoading()) \
@@ -176,6 +232,32 @@
 			(JsonSerializableObject).Serialize(Serializer, true); \
 			Serializer.EndObject(); \
 		}
+		
+#define JSON_SERIALIZE_OBJECT_SERIALIZABLE_WITHDEFAULT(JsonName, JsonSerializableObject, JsonSerializableObjectDefault) \
+		/* Process the JsonName field differently because it is an object */ \
+		if (Serializer.IsLoading()) \
+		{ \
+			/* Read in the value from the JsonName field */ \
+			if (Serializer.GetObject()->HasTypedField<EJson::Object>(JsonName)) \
+			{ \
+				TSharedPtr<FJsonObject> JsonObj = Serializer.GetObject()->GetObjectField(JsonName); \
+				if (JsonObj.IsValid()) \
+				{ \
+					(JsonSerializableObject).FromJson(JsonObj); \
+				} \
+			} \
+			else \
+			{ \
+				JsonSerializableObject = JsonSerializableObjectDefault; \
+			} \
+		} \
+		else \
+		{ \
+			/* Write the value to the Name field */ \
+			Serializer.StartObject(JsonName); \
+			(JsonSerializableObject).Serialize(Serializer, true); \
+			Serializer.EndObject(); \
+		}
 
 #define JSON_SERIALIZE_OPTIONAL_OBJECT_SERIALIZABLE(JsonName, JsonSerializableObject) \
 		if (Serializer.IsLoading()) \
@@ -269,6 +351,7 @@ struct FJsonSerializerBase;
 /** Array of data */
 typedef TArray<FString> FJsonSerializableArray;
 typedef TArray<int32> FJsonSerializableArrayInt;
+typedef TArray<float> FJsonSerializableArrayFloat;
 
 /** Maps a key to a value */
 typedef TMap<FString, FString> FJsonSerializableKeyValueMap;
@@ -301,6 +384,7 @@ struct FJsonSerializerBase
 	virtual void SerializeArray(FJsonSerializableArray& Array) = 0;
 	virtual void SerializeArray(const TCHAR* Name, FJsonSerializableArray& Value) = 0;
 	virtual void SerializeArray(const TCHAR* Name, FJsonSerializableArrayInt& Value) = 0;
+	virtual void SerializeArray(const TCHAR* Name, FJsonSerializableArrayFloat& Value) = 0;
 	virtual void SerializeMap(const TCHAR* Name, FJsonSerializableKeyValueMap& Map) = 0;
 	virtual void SerializeMap(const TCHAR* Name, FJsonSerializableKeyValueMapInt& Map) = 0;
 	virtual void SerializeMap(const TCHAR* Name, FJsonSerializableKeyValueMapInt64& Map) = 0;
@@ -524,6 +608,23 @@ public:
 		JsonWriter->WriteArrayEnd();
 	}
 
+	/**
+	 * Serializes an array of values with an identifier
+	 *
+	 * @param Name the name of the property to serialize
+	 * @param Array the array to serialize
+	 */
+	virtual void SerializeArray(const TCHAR* Name, FJsonSerializableArrayFloat& Array) override
+	{
+		JsonWriter->WriteArrayStart(Name);
+		// Iterate all of values
+		for (FJsonSerializableArrayFloat::ElementType& Item : Array)
+		{
+			JsonWriter->WriteValue(Item);
+		}
+		JsonWriter->WriteArrayEnd();
+	}
+
 	/**
 	 * Serializes the keys & values for map
 	 *
@@ -845,6 +946,24 @@ public:
 			}
 		}
 	}
+	/**
+	 * Serializes an array of values with an identifier
+	 *
+	 * @param Name the name of the property to serialize
+	 * @param Array the array to serialize
+	 */
+	virtual void SerializeArray(const TCHAR* Name, FJsonSerializableArrayFloat& Array) override
+	{
+		if (JsonObject->HasTypedField<EJson::Array>(Name))
+		{
+			TArray< TSharedPtr<FJsonValue> > JsonArray = JsonObject->GetArrayField(Name);
+			// Iterate all of the keys and their values
+			for (TSharedPtr<FJsonValue>& Value : JsonArray)
+			{
+				Array.Add((float)Value->AsNumber());
+			}
+		}
+	}
 	/**
 	 * Serializes the keys & values for map
 	 *
