JsonSerializerMacros: adding SerializeMap for FJsonSerializableArrayInt values

diff --git a/Engine/Source/Runtime/Json/Public/Serialization/JsonSerializerMacros.h b/Engine/Source/Runtime/Json/Public/Serialization/JsonSerializerMacros.h
--- a/Engine/Source/Runtime/Json/Public/Serialization/JsonSerializerMacros.h
+++ b/Engine/Source/Runtime/Json/Public/Serialization/JsonSerializerMacros.h
@@ -358,6 +358,7 @@ typedef TMap<FString, FString> FJsonSerializableKeyValueMap;
 typedef TMap<FString, int32> FJsonSerializableKeyValueMapInt;
 typedef TMap<FString, int64> FJsonSerializableKeyValueMapInt64;
 typedef TMap<FString, float> FJsonSerializableKeyValueMapFloat;
+typedef TMap<FString, FJsonSerializableArrayInt> FJsonSerializableKeyValueMapArrayInt;
 
 /**
  * Base interface used to serialize to/from JSON. Hides the fact there are separate read/write classes
@@ -387,6 +388,7 @@ struct FJsonSerializerBase
 	virtual void SerializeArray(const TCHAR* Name, FJsonSerializableArrayFloat& Value) = 0;
 	virtual void SerializeMap(const TCHAR* Name, FJsonSerializableKeyValueMap& Map) = 0;
 	virtual void SerializeMap(const TCHAR* Name, FJsonSerializableKeyValueMapInt& Map) = 0;
+	virtual void SerializeMap(const TCHAR* Name, FJsonSerializableKeyValueMapArrayInt& Map) = 0;
 	virtual void SerializeMap(const TCHAR* Name, FJsonSerializableKeyValueMapInt64& Map) = 0;
 	virtual void SerializeMap(const TCHAR* Name, FJsonSerializableKeyValueMapFloat& Map) = 0;
 	virtual void SerializeSimpleMap(FJsonSerializableKeyValueMap& Map) = 0;
@@ -658,6 +660,22 @@ public:
 		}
 		JsonWriter->WriteObjectEnd();
 	}
+	/**
+	 * Serializes the keys & values for map
+	 *
+	 * @param Name the name of the property to serialize
+	 * @param Map the map to serialize
+	 */
+	virtual void SerializeMap(const TCHAR* Name, FJsonSerializableKeyValueMapArrayInt& Map) override
+	{
+		JsonWriter->WriteObjectStart(Name);
+		// Iterate all of the keys and their values
+		for (FJsonSerializableKeyValueMapArrayInt::ElementType& Pair : Map)
+		{
+			SerializeArray(*Pair.Key, Pair.Value);
+		}
+		JsonWriter->WriteObjectEnd();
+	}
 
 	/**
 	 * Serializes the keys & values for map
@@ -1002,6 +1020,34 @@ public:
 			}
 		}
 	}
+	
+	/**
+     * Serializes the keys & values for map
+     *
+     * @param Name the name of the property to serialize
+     * @param Map the map to serialize
+     */
+    virtual void SerializeMap(const TCHAR* Name, FJsonSerializableKeyValueMapArrayInt& Map) override
+    {
+    	if (JsonObject->HasTypedField<EJson::Object>(Name))
+    	{
+    		TSharedPtr<FJsonObject> JsonMap = JsonObject->GetObjectField(Name);
+    		// Iterate all of the keys and their values
+    		for (const TPair<FString, TSharedPtr<FJsonValue>>& Pair : JsonMap->Values)
+    		{
+    			FJsonSerializableArrayInt TempArray;
+    			for(const TSharedPtr<FJsonValue> ArrayValue : Pair.Value->AsArray())
+    			{
+    				int32 IntValue;
+    				if(ArrayValue.IsValid() && ArrayValue->TryGetNumber(IntValue))
+    				{
+    					TempArray.Add(IntValue);
+    				}
+    			}
+    			Map.Add(Pair.Key, TempArray);
+    		}
+    	}
+    }
 
 	/**
 	 * Serializes the keys & values for map
