JsonSerializerMacros: fix for iterating over array without using by ref

diff --git a/Engine/Source/Runtime/Json/Public/Serialization/JsonSerializerMacros.h b/Engine/Source/Runtime/Json/Public/Serialization/JsonSerializerMacros.h
--- a/Engine/Source/Runtime/Json/Public/Serialization/JsonSerializerMacros.h
+++ b/Engine/Source/Runtime/Json/Public/Serialization/JsonSerializerMacros.h
@@ -1036,7 +1036,7 @@ public:
     		for (const TPair<FString, TSharedPtr<FJsonValue>>& Pair : JsonMap->Values)
     		{
     			FJsonSerializableArrayInt TempArray;
-    			for(const TSharedPtr<FJsonValue> ArrayValue : Pair.Value->AsArray())
+    			for(const TSharedPtr<FJsonValue>& ArrayValue : Pair.Value->AsArray())
     			{
     				int32 IntValue;
     				if(ArrayValue.IsValid() && ArrayValue->TryGetNumber(IntValue))
