- Add JsonValue Deserialize taking const object

diff --git a/Engine/Source/Runtime/Json/Private/Dom/JsonValue.cpp b/Engine/Source/Runtime/Json/Private/Dom/JsonValue.cpp
--- a/Engine/Source/Runtime/Json/Private/Dom/JsonValue.cpp
+++ b/Engine/Source/Runtime/Json/Private/Dom/JsonValue.cpp
@@ -274,6 +274,11 @@ static void DuplicateJsonArray(const TArray<TSharedPtr<FJsonValue>>& Source, TAr
 	} 
 }
 
+TSharedPtr<FJsonValue> FJsonValue::Duplicate(const TSharedPtr<const FJsonValue>& Src)
+{
+	return Duplicate(ConstCastSharedPtr<FJsonValue>(Src));
+}
+
 TSharedPtr<FJsonValue> FJsonValue::Duplicate(const TSharedPtr<FJsonValue>& Src)
 {
 	switch (Src->Type)
diff --git a/Engine/Source/Runtime/Json/Public/Dom/JsonValue.h b/Engine/Source/Runtime/Json/Public/Dom/JsonValue.h
--- a/Engine/Source/Runtime/Json/Public/Dom/JsonValue.h
+++ b/Engine/Source/Runtime/Json/Public/Dom/JsonValue.h
@@ -98,6 +98,7 @@ public:
 
 	EJson Type;
 
+	static JSON_API TSharedPtr<FJsonValue> Duplicate(const TSharedPtr<const FJsonValue>& Src);
 	static JSON_API TSharedPtr<FJsonValue> Duplicate(const TSharedPtr<FJsonValue>& Src);
 
 	static JSON_API bool CompareEqual(const FJsonValue& Lhs, const FJsonValue& Rhs);
