[Json] Handle empty keys inside json objects.

Previously we would try to ignore empty string keys and serialize the values without it but that would trigger an assert inside json objects because only arrays can serialize values without keys. Empty strings are valid key strings in Json, so the previous implementation did not respect the specifications.

This change makes sure we only support empty keys when we're serializing a json object (arrays are excluded). Automated test cases have also been added to validate the empty key support and that we properly handle the case that was causing a crash during previous attempt at this fix (see backout CL 23028888).

diff --git a/Engine/Source/Runtime/Json/Private/Tests/JsonTests.cpp b/Engine/Source/Runtime/Json/Private/Tests/JsonTests.cpp
--- a/Engine/Source/Runtime/Json/Private/Tests/JsonTests.cpp
+++ b/Engine/Source/Runtime/Json/Private/Tests/JsonTests.cpp
@@ -69,6 +69,29 @@ bool FJsonAutomationTest::RunTest(const FString& Parameters)
 		check( InputString == OutputString );
 	}
 
+	// Empty Array with Empty Identifier Case
+	{
+		const FString ExpectedString = TEXT("[]");
+		FString OutputString;
+		TSharedRef<FJsonValueArray> EmptyValuesArray = MakeShared<FJsonValueArray>(TArray<TSharedPtr<FJsonValue>>());
+		TSharedRef<FCondensedJsonStringWriter> JsonWriter = FCondensedJsonStringWriterFactory::Create( &OutputString );
+		verify(FJsonSerializer::Serialize(EmptyValuesArray, FString(), JsonWriter));
+		check(ExpectedString == OutputString);
+	}
+
+	// Serializing Object Value with Empty Identifier Case
+	{
+		const FString ExpectedString = TEXT("{\"\":\"foo\"}");
+		FString OutputString;
+		TSharedRef<FJsonValue> FooValue = MakeShared<FJsonValueString>("foo");
+		TSharedRef<FCondensedJsonStringWriter> JsonWriter = FCondensedJsonStringWriterFactory::Create(&OutputString);
+		JsonWriter->WriteObjectStart();
+		verify(FJsonSerializer::Serialize(FooValue, FString(), JsonWriter, false));
+		JsonWriter->WriteObjectEnd();
+		JsonWriter->Close();
+		check(ExpectedString == OutputString);
+	}
+
 	// Simple Array Case
 	{
 		const FString InputString = 
@@ -234,6 +257,7 @@ bool FJsonAutomationTest::RunTest(const FString& Parameters)
 					"\"Some String1\","
 					"10,"
 					"{"
+						"\"\":\"Empty Key\","
 						"\"Value\":\"Some String3\""
 					"},"
 					"["
@@ -267,6 +291,7 @@ bool FJsonAutomationTest::RunTest(const FString& Parameters)
 		TSharedPtr< FJsonObject > Object = Array[2]->AsObject();
 		check(Object.IsValid());
 		check(Object->GetStringField(TEXT("Value")) == TEXT("Some String3"));
+		check(Object->GetStringField(TEXT("")) == TEXT("Empty Key"));
 
 		const TArray<TSharedPtr< FJsonValue >>& InnerArray = Array[3]->AsArray();
 		check(InnerArray.Num() == 2);
diff --git a/Engine/Source/Runtime/Json/Public/Serialization/JsonSerializer.h b/Engine/Source/Runtime/Json/Public/Serialization/JsonSerializer.h
--- a/Engine/Source/Runtime/Json/Public/Serialization/JsonSerializer.h
+++ b/Engine/Source/Runtime/Json/Public/Serialization/JsonSerializer.h
@@ -105,12 +105,38 @@ public:
 		return true;
 	}
 
+	/**
+	 * Serialize the passed array of json values into the writer.
+	 * This will effectively serialize all of the values enclosed in [] square brackets.
+	 * Example:
+	 *  - Writer state: [123 <writer position>
+	 *    Parameter: Array: ["foo", "bar", "", 456]
+	 *    Serialization result: [123, ["foo", "bar", "", 456] <writer position>
+	 *
+	 * @param Array		    The json array we are serializing
+	 * @param Writer		The writer the array is written into.
+	 * @param bCloseWriter	When set to true the Writer will be closed after the serialization.
+	 * @return				Returns true if the serialization was successful, false otherwise.
+	 */
 	template <class CharType, class PrintPolicy>
 	static bool Serialize(const TArray<TSharedPtr<FJsonValue>>& Array, const TSharedRef<TJsonWriter<CharType, PrintPolicy>>& Writer, bool bCloseWriter = true)
 	{
 		return Serialize(Array, *Writer, bCloseWriter);
 	}
 
+	/**
+	 * Serialize the passed array of json values into the writer.
+	 * This will effectively serialize all of the values enclosed in [] square brackets.
+	 * Example:
+	 *  - Writer state: [123 <writer position>
+	 *    Parameter: Array: ["foo", "bar", "", 456]
+	 *    Serialization result: [123, ["foo", "bar", "", 456] <writer position>
+	 *
+	 * @param Array		    The json array we are serializing
+	 * @param Writer		The writer the array is written into.
+	 * @param bCloseWriter	When set to true the Writer will be closed after the serialization.
+	 * @return				Returns true if the serialization was successful, false otherwise.
+	 */
 	template <class CharType, class PrintPolicy>
 	static bool Serialize(const TArray<TSharedPtr<FJsonValue>>& Array, TJsonWriter<CharType, PrintPolicy>& Writer, bool bCloseWriter = true )
 	{
@@ -118,12 +144,38 @@ public:
 		return FJsonSerializer::Serialize<CharType, PrintPolicy>(StartingElement, Writer, bCloseWriter);
 	}
 
+	/**
+	 * Serialize the passed Json object into the writer.
+	 * This will effectively serialize all of the identifier:value pairs of the object enclosed in {} curly brackets.
+	 * Example:
+	 *  - Writer state: [123 <writer position>
+	 *    Parameter: Object: {"foo": "bar", "baz": "", "": 456}
+	 *    Serialization result: [123, {"foo": "bar", "baz": "", "": 456} <writer position>
+	 *
+	 * @param Object		The json object we are serializing
+	 * @param Writer		The writer the object is written into.
+	 * @param bCloseWriter	When set to true the Writer will be closed after the serialization.
+	 * @return				Returns true if the serialization was successful, false otherwise.
+	 */
 	template <class CharType, class PrintPolicy>
 	static bool Serialize(const TSharedRef<FJsonObject>& Object, const TSharedRef<TJsonWriter<CharType, PrintPolicy>>& Writer, bool bCloseWriter = true )
 	{
 		return Serialize(Object, *Writer, bCloseWriter);
 	}
 
+	/**
+	 * Serialize the passed Json object into the writer.
+	 * This will effectively serialize all of the identifier:value pairs of the object enclosed in {} brackets.
+	 * Example:
+	 *  - Writer state: [123 <writer position>
+	 *    Parameter: Object: {"foo": "bar", "baz": "", "": 456}
+	 *    Serialization result: [123, {"foo": "bar", "baz": "", "": 456} <writer position>
+	 * 
+     * @param Object		The json object we are serializing
+	 * @param Writer		The writer the object is written into.
+	 * @param bCloseWriter	When set to true the Writer will be closed after the serialization.
+	 * @return				Returns true if the serialization was successful, false otherwise.
+	 */
 	template <class CharType, class PrintPolicy>
 	static bool Serialize(const TSharedRef<FJsonObject>& Object, TJsonWriter<CharType, PrintPolicy>& Writer, bool bCloseWriter = true)
 	{
@@ -131,11 +183,54 @@ public:
 		return FJsonSerializer::Serialize<CharType, PrintPolicy>(StartingElement, Writer, bCloseWriter);
 	}
 
+	/**
+	 * Serialize the passed Json value and identifier into the writer.
+	 * Empty string identifiers will be ignored when the writer is not writing inside of a json object and only the value will be serialized.
+	 * If the writer is in a state where it's currently writing inside of a json object, then the identifier will always be serialized.
+	 * Examples:
+	 *	- Writer state: { "foo": "bar" <writer position>
+	 *    Parameters: Identifier: ""
+	 *                Value: "baz"
+	 *    Serialization result: { "foo": "bar", "": "baz" <writer position> //empty identifier is serialized as a valid key for the key:value pair "":"baz"
+	 *
+	 * - Writer state: { "foo": ["bar" <writer position>
+	 *   Parameters: Identifier: ""
+	 *               Value: "baz"
+	 *   Serialization result: { foo: ["bar", "baz" <writer position> //empty identifier is ignored since we are writing into an array and not an object.
+	 *
+	 * @param Value			The json value we are serializing
+	 * @param Identifier	The identifier of the value, empty identifiers are ignored outside of json objects.
+	 * @param Writer		The writer the value and identifier are written into.
+	 * @param bCloseWriter	When set to true the Writer will be closed after the serialization.
+	 * @return				Returns true if the serialization was successful, false otherwise.
+	 */
 	template <class CharType, class PrintPolicy>
 	static bool Serialize(const TSharedPtr<FJsonValue>& Value, const FString& Identifier, const TSharedRef<TJsonWriter<CharType, PrintPolicy>>& Writer, bool bCloseWriter = true)
 	{
 		return Serialize(Value, Identifier, *Writer, bCloseWriter);
 	}
+
+	/**
+	 * Serialize the passed Json value and identifier into the writer.
+	 * Empty string identifiers will be ignored when the writer is not writing inside of a json object and only the value will be serialized.
+	 * If the writer is in a state where it's currently writing inside of a json object, then the identifier will always be serialized.
+	 * Examples:
+	 *	- Writer state: { "foo": "bar" <writer position>
+	 *    Parameters: Identifier: ""
+	 *                Value: "baz"
+	 *    Serialization result: { "foo": "bar", "": "baz" <writer position> //empty identifier is serialized as a valid key for the key:value pair "":"baz"
+	 * 
+	 * - Writer state: { "foo": ["bar" <writer position>
+	 *   Parameters: Identifier: ""
+	 *               Value: "baz"
+	 *   Serialization result: { foo: ["bar", "baz" <writer position> //empty identifier is ignored since we are writing into an array and not an object.
+	 * 
+	 * @param Value			The json value we are serializing
+	 * @param Identifier	The identifier of the value, empty identifiers are ignored outside of json objects.
+	 * @param Writer		The writer the value and identifier are written into.
+	 * @param bCloseWriter	When set to true the Writer will be closed after the serialization.
+	 * @return				Returns true if the serialization was successful, false otherwise.
+	 */
 	template <class CharType, class PrintPolicy>
 	static bool Serialize(const TSharedPtr<FJsonValue>& Value, const FString& Identifier, TJsonWriter<CharType, PrintPolicy>& Writer, bool bCloseWriter = true)
 	{
@@ -158,32 +253,30 @@ private:
 		FElement( const TSharedPtr<FJsonValue>& InValue )
 			: Identifier()
 			, Value(InValue)
-			, HasBeenProcessed(false)
 		{ }
 
 		FElement( const TSharedRef<FJsonObject>& Object )
 			: Identifier()
 			, Value(MakeShared<FJsonValueObject>(Object))
-			, HasBeenProcessed( false )
 		{ }
 
 		FElement( const TArray<TSharedPtr<FJsonValue>>& Array )
 			: Identifier()
 			, Value(MakeShared<FJsonValueArray>(Array))
-			, HasBeenProcessed(false)
 		{ }
 
 		FElement( const FString& InIdentifier, const TSharedPtr< FJsonValue >& InValue )
 			: Identifier( InIdentifier )
 			, Value( InValue )
-			, HasBeenProcessed( false )
+			, bIsKeyValuePair( true )
 		{
 
 		}
 
 		FString Identifier;
 		TSharedPtr< FJsonValue > Value;
-		bool HasBeenProcessed;
+		bool bHasBeenProcessed = false;
+		bool bIsKeyValuePair = false;
 	};
 
 private:
@@ -313,13 +406,15 @@ private:
 		while (ElementStack.Num() > 0)
 		{
 			TSharedRef<FElement> Element = ElementStack.Pop();
+			// Empty keys are valid identifiers only when writing inside an object.
+			const bool bWriteValueOnly = !Element->bIsKeyValuePair || (Element->Identifier.IsEmpty() && Writer.GetCurrentElementType() != EJson::Object);
 			check(Element->Value->Type != EJson::None);
 
 			switch (Element->Value->Type)
 			{
 			case EJson::Number:	
 				{
-					if (Element->Identifier.IsEmpty())
+					if (bWriteValueOnly)
 					{
 						if (Element->Value->PreferStringRepresentation())
 						{
@@ -346,7 +441,7 @@ private:
 
 			case EJson::Boolean:					
 				{
-					if (Element->Identifier.IsEmpty())
+					if (bWriteValueOnly)
 					{
 						Writer.WriteValue(Element->Value->AsBool());
 					}
@@ -359,7 +454,7 @@ private:
 
 			case EJson::String:
 				{
-					if (Element->Identifier.IsEmpty())
+					if (bWriteValueOnly)
 					{
 						Writer.WriteValue(Element->Value->AsString());
 					}
@@ -372,7 +467,7 @@ private:
 
 			case EJson::Null:
 				{
-					if (Element->Identifier.IsEmpty())
+					if (bWriteValueOnly)
 					{
 						Writer.WriteNull();
 					}
@@ -385,16 +480,16 @@ private:
 
 			case EJson::Array:
 				{
-					if (Element->HasBeenProcessed)
+					if (Element->bHasBeenProcessed)
 					{
 						Writer.WriteArrayEnd();
 					}
 					else
 					{
-						Element->HasBeenProcessed = true;
+						Element->bHasBeenProcessed = true;
 						ElementStack.Push(Element);
 
-						if (Element->Identifier.IsEmpty())
+						if (bWriteValueOnly)
 						{
 							Writer.WriteArrayStart();
 						}
@@ -415,16 +510,16 @@ private:
 
 			case EJson::Object:
 				{
-					if (Element->HasBeenProcessed)
+					if (Element->bHasBeenProcessed)
 					{
 						Writer.WriteObjectEnd();
 					}
 					else
 					{
-						Element->HasBeenProcessed = true;
+						Element->bHasBeenProcessed = true;
 						ElementStack.Push(Element);
 
-						if (Element->Identifier.IsEmpty())
+						if (bWriteValueOnly)
 						{
 							Writer.WriteObjectStart();
 						}
diff --git a/Engine/Source/Runtime/Json/Public/Serialization/JsonWriter.h b/Engine/Source/Runtime/Json/Public/Serialization/JsonWriter.h
--- a/Engine/Source/Runtime/Json/Public/Serialization/JsonWriter.h
+++ b/Engine/Source/Runtime/Json/Public/Serialization/JsonWriter.h
@@ -87,6 +87,11 @@ public:
 		return CanWriteObjectWithoutIdentifier();
 	}
 
+	EJson GetCurrentElementType() const
+	{
+		return Stack.Num() > 0 ? Stack.Top() : EJson::None;
+	}
+
 	void WriteObjectStart()
 	{
 		check(CanWriteObjectWithoutIdentifier());
