Try to fix the live random crash which might be related to big chunk of memory allocation.

diff --git a/Engine/Source/Runtime/Json/Public/Serialization/JsonSerializerMacros.h b/Engine/Source/Runtime/Json/Public/Serialization/JsonSerializerMacros.h
--- a/Engine/Source/Runtime/Json/Public/Serialization/JsonSerializerMacros.h
+++ b/Engine/Source/Runtime/Json/Public/Serialization/JsonSerializerMacros.h
@@ -1225,6 +1225,22 @@ struct FJsonSerializable
 		return false;
 	}
 
+	template <class CharType = TCHAR>
+	bool FromJson(TStringView<CharType> Json)
+	{
+		TSharedPtr<FJsonObject> JsonObject;
+		TSharedRef<TJsonReader<CharType> > JsonReader = TJsonReaderFactory<CharType>::CreateFromView(Json);
+		if (FJsonSerializer::Deserialize(JsonReader,JsonObject) &&
+			JsonObject.IsValid())
+		{
+			FJsonSerializerReader Serializer(JsonObject);
+			Serialize(Serializer, false);
+			return true;
+		}
+		UE_LOG(LogJson, Warning, TEXT("Failed to parse Json from a string: %s"), *JsonReader->GetErrorMessage());
+		return false;
+	}
+
 	virtual bool FromJson(TSharedPtr<FJsonObject> JsonObject)
 	{
 		if (JsonObject.IsValid())
