[JSON] JsonObject, JsonWriter and JsonSerializerMacros updated to use FStringViews and FString

diff --git a/Engine/Source/Runtime/Json/Private/Dom/JsonObject.cpp b/Engine/Source/Runtime/Json/Private/Dom/JsonObject.cpp
--- a/Engine/Source/Runtime/Json/Private/Dom/JsonObject.cpp
+++ b/Engine/Source/Runtime/Json/Private/Dom/JsonObject.cpp
@@ -2,113 +2,108 @@
 
 #include "Dom/JsonObject.h"
 
-
-void FJsonObject::SetField( const FString& FieldName, const TSharedPtr<FJsonValue>& Value )
+void FJsonObject::SetField(FString&& FieldName, const TSharedPtr<FJsonValue>& Value)
 {
-	this->Values.Add(FieldName, Value);
+	this->Values.Add(MoveTemp(FieldName), Value);
 }
 
-
-void FJsonObject::RemoveField( const FString& FieldName )
+void FJsonObject::SetField(const FString& FieldName, const TSharedPtr<FJsonValue>& Value)
 {
-	this->Values.Remove(FieldName);
+	SetField(CopyTemp(FieldName), Value);
 }
 
+void FJsonObject::RemoveField(FStringView FieldName)
+{
+	this->Values.RemoveByHash(GetTypeHash(FieldName), FieldName);
+}
 
-double FJsonObject::GetNumberField( const FString& FieldName ) const
+double FJsonObject::GetNumberField(FStringView FieldName) const
 {
 	return GetField<EJson::None>(FieldName)->AsNumber();
 }
 
-
-bool FJsonObject::TryGetNumberField( const FString& FieldName, float& OutNumber ) const
+bool FJsonObject::TryGetNumberField(FStringView FieldName, float& OutNumber) const
 {
 	TSharedPtr<FJsonValue> Field = TryGetField(FieldName);
 	return Field.IsValid() && Field->TryGetNumber(OutNumber);
 }
 
-
-bool FJsonObject::TryGetNumberField( const FString& FieldName, double& OutNumber ) const
+bool FJsonObject::TryGetNumberField(FStringView FieldName, double& OutNumber) const
 {
 	TSharedPtr<FJsonValue> Field = TryGetField(FieldName);
 	return Field.IsValid() && Field->TryGetNumber(OutNumber);
 }
 
-
-bool FJsonObject::TryGetNumberField( const FString& FieldName, int8& OutNumber ) const
+bool FJsonObject::TryGetNumberField(FStringView FieldName, int8& OutNumber) const
 {
 	TSharedPtr<FJsonValue> Field = TryGetField(FieldName);
 	return Field.IsValid() && Field->TryGetNumber(OutNumber);
 }
 
-
-bool FJsonObject::TryGetNumberField( const FString& FieldName, int16& OutNumber ) const
+bool FJsonObject::TryGetNumberField(FStringView FieldName, int16& OutNumber) const
 {
 	TSharedPtr<FJsonValue> Field = TryGetField(FieldName);
 	return Field.IsValid() && Field->TryGetNumber(OutNumber);
 }
 
-
-bool FJsonObject::TryGetNumberField( const FString& FieldName, int32& OutNumber ) const
+bool FJsonObject::TryGetNumberField(FStringView FieldName, int32& OutNumber) const
 {
 	TSharedPtr<FJsonValue> Field = TryGetField(FieldName);
 	return Field.IsValid() && Field->TryGetNumber(OutNumber);
 }
 
-
-bool FJsonObject::TryGetNumberField(const FString& FieldName, int64& OutNumber) const
+bool FJsonObject::TryGetNumberField(FStringView FieldName, int64& OutNumber) const
 {
 	TSharedPtr<FJsonValue> Field = TryGetField(FieldName);
 	return Field.IsValid() && Field->TryGetNumber(OutNumber);
 }
 
-
-bool FJsonObject::TryGetNumberField( const FString& FieldName, uint8& OutNumber ) const
+bool FJsonObject::TryGetNumberField(FStringView FieldName, uint8& OutNumber) const
 {
 	TSharedPtr<FJsonValue> Field = TryGetField(FieldName);
 	return Field.IsValid() && Field->TryGetNumber(OutNumber);
 }
 
-
-bool FJsonObject::TryGetNumberField( const FString& FieldName, uint16& OutNumber ) const
+bool FJsonObject::TryGetNumberField(FStringView FieldName, uint16& OutNumber) const
 {
 	TSharedPtr<FJsonValue> Field = TryGetField(FieldName);
 	return Field.IsValid() && Field->TryGetNumber(OutNumber);
 }
 
-bool FJsonObject::TryGetNumberField( const FString& FieldName, uint32& OutNumber ) const
+bool FJsonObject::TryGetNumberField(FStringView FieldName, uint32& OutNumber) const
 {
 	TSharedPtr<FJsonValue> Field = TryGetField(FieldName);
 	return Field.IsValid() && Field->TryGetNumber(OutNumber);
 }
 
-bool FJsonObject::TryGetNumberField( const FString& FieldName, uint64& OutNumber ) const
+bool FJsonObject::TryGetNumberField(FStringView FieldName, uint64& OutNumber) const
 {
 	TSharedPtr<FJsonValue> Field = TryGetField(FieldName);
 	return Field.IsValid() && Field->TryGetNumber(OutNumber);
 }
 
-
-void FJsonObject::SetNumberField( const FString& FieldName, double Number )
+void FJsonObject::SetNumberField(FString&& FieldName, double Number)
 {
-	this->Values.Add(FieldName, MakeShared<FJsonValueNumber>(Number));
+	this->Values.Add(MoveTemp(FieldName), MakeShared<FJsonValueNumber>(Number));
 }
 
+void FJsonObject::SetNumberField(const FString& FieldName, double Number)
+{
+	SetNumberField(CopyTemp(FieldName), Number);
+}
 
-FString FJsonObject::GetStringField( const FString& FieldName ) const
+FString FJsonObject::GetStringField(FStringView FieldName) const
 {
 	return GetField<EJson::None>(FieldName)->AsString();
 }
 
-
-bool FJsonObject::TryGetStringField( const FString& FieldName, FString& OutString ) const
+bool FJsonObject::TryGetStringField(FStringView FieldName, FString& OutString) const
 {
 	TSharedPtr<FJsonValue> Field = TryGetField(FieldName);
 	return Field.IsValid() && Field->TryGetString(OutString);
 }
 
-
-bool FJsonObject::TryGetStringArrayField( const FString& FieldName, TArray<FString>& OutArray ) const
+bool FJsonObject::TryGetStringArrayField(FStringView FieldName, TArray<FString>& OutArray) const
 {
 	TSharedPtr<FJsonValue> Field = TryGetField(FieldName);
 
@@ -139,76 +134,106 @@ bool FJsonObject::TryGetStringArrayField( const FString& FieldName, TArray<FStri
 	return true;
 }
 
+void FJsonObject::SetStringField(FString&& FieldName, FString&& StringValue)
+{
+	this->Values.Add(MoveTemp(FieldName), MakeShared<FJsonValueString>(MoveTemp(StringValue)));
+}
 
-void FJsonObject::SetStringField( const FString& FieldName, const FString& StringValue )
+void FJsonObject::SetStringField(FString&& FieldName, const FString& StringValue)
 {
-	this->Values.Add(FieldName, MakeShared<FJsonValueString>(StringValue));
+	SetStringField(MoveTemp(FieldName), CopyTemp(StringValue));
 }
 
+void FJsonObject::SetStringField(const FString& FieldName, FString&& StringValue)
+{
+	SetStringField(CopyTemp(FieldName), MoveTemp(StringValue));
+}
 
-bool FJsonObject::GetBoolField( const FString& FieldName ) const
+void FJsonObject::SetStringField(const FString& FieldName, const FString& StringValue)
 {
-	return GetField<EJson::None>(FieldName)->AsBool();
+	SetStringField(CopyTemp(FieldName), CopyTemp(StringValue));
 }
 
+bool FJsonObject::GetBoolField(FStringView FieldName) const
+{
+	return GetField<EJson::None>(FieldName)->AsBool();
+}
 
-bool FJsonObject::TryGetBoolField( const FString& FieldName, bool& OutBool ) const
+bool FJsonObject::TryGetBoolField(FStringView FieldName, bool& OutBool) const
 {
 	TSharedPtr<FJsonValue> Field = TryGetField(FieldName);
 	return Field.IsValid() && Field->TryGetBool(OutBool);
 }
 
-
-void FJsonObject::SetBoolField( const FString& FieldName, bool InValue )
+void FJsonObject::SetBoolField(FString&& FieldName, bool InValue)
 {
-	this->Values.Add(FieldName, MakeShared<FJsonValueBoolean>(InValue));
+	this->Values.Add(MoveTemp(FieldName), MakeShared<FJsonValueBoolean>(InValue));
 }
 
+void FJsonObject::SetBoolField(const FString& FieldName, bool InValue)
+{
+	SetBoolField(CopyTemp(FieldName), InValue);
+}
 
-const TArray<TSharedPtr<FJsonValue>>& FJsonObject::GetArrayField( const FString& FieldName ) const
+const TArray<TSharedPtr<FJsonValue>>& FJsonObject::GetArrayField(FStringView FieldName) const
 {
 	return GetField<EJson::Array>(FieldName)->AsArray();
 }
 
-
-bool FJsonObject::TryGetArrayField(const FString& FieldName, const TArray< TSharedPtr<FJsonValue> >*& OutArray) const
+bool FJsonObject::TryGetArrayField(FStringView FieldName, const TArray<TSharedPtr<FJsonValue>>*& OutArray) const
 {
 	TSharedPtr<FJsonValue> Field = TryGetField(FieldName);
 	return Field.IsValid() && Field->TryGetArray(OutArray);
 }
 
+void FJsonObject::SetArrayField(FString&& FieldName, TArray<TSharedPtr<FJsonValue>>&& Array)
+{
+	this->Values.Add(MoveTemp(FieldName), MakeShared<FJsonValueArray>(MoveTemp(Array)));
+}
+
+void FJsonObject::SetArrayField(FString&& FieldName, const TArray<TSharedPtr<FJsonValue>>& Array)
+{
+	SetArrayField(MoveTemp(FieldName), CopyTemp(Array));
+}
 
-void FJsonObject::SetArrayField( const FString& FieldName, const TArray< TSharedPtr<FJsonValue> >& Array )
+void FJsonObject::SetArrayField(const FString& FieldName, TArray<TSharedPtr<FJsonValue>>&& Array)
 {
-	this->Values.Add(FieldName, MakeShared<FJsonValueArray>(Array));
+	SetArrayField(CopyTemp(FieldName), MoveTemp(Array));
 }
 
+void FJsonObject::SetArrayField(const FString& FieldName, const TArray<TSharedPtr<FJsonValue>>& Array)
+{
+	SetArrayField(CopyTemp(FieldName), CopyTemp(Array));
+}
 
-const TSharedPtr<FJsonObject>& FJsonObject::GetObjectField( const FString& FieldName ) const
+const TSharedPtr<FJsonObject>& FJsonObject::GetObjectField(FStringView FieldName) const
 {
 	return GetField<EJson::Object>(FieldName)->AsObject();
 }
 
-
-bool FJsonObject::TryGetObjectField( const FString& FieldName, const TSharedPtr<FJsonObject>*& OutObject ) const
+bool FJsonObject::TryGetObjectField(FStringView FieldName, const TSharedPtr<FJsonObject>*& OutObject) const
 {
 	TSharedPtr<FJsonValue> Field = TryGetField(FieldName);
 	return Field.IsValid() && Field->TryGetObject(OutObject);
 }
 
-
-void FJsonObject::SetObjectField( const FString& FieldName, const TSharedPtr<FJsonObject>& JsonObject )
+void FJsonObject::SetObjectField(FString&& FieldName, const TSharedPtr<FJsonObject>& JsonObject)
 {
 	if (JsonObject.IsValid())
 	{
-		this->Values.Add(FieldName, MakeShared<FJsonValueObject>(JsonObject.ToSharedRef()));
+		this->Values.Add(MoveTemp(FieldName), MakeShared<FJsonValueObject>(JsonObject.ToSharedRef()));
 	}
 	else
 	{
-		this->Values.Add(FieldName, MakeShared<FJsonValueNull>());
+		this->Values.Add(MoveTemp(FieldName), MakeShared<FJsonValueNull>());
 	}
 }
 
+void FJsonObject::SetObjectField(const FString& FieldName, const TSharedPtr<FJsonObject>& JsonObject)
+{
+	SetObjectField(CopyTemp(FieldName), JsonObject);
+}
+
 void FJsonObject::Duplicate(const TSharedPtr<const FJsonObject>& Source, const TSharedPtr<FJsonObject>& Dest)
 {
 	if (Source && Dest)
diff --git a/Engine/Source/Runtime/Json/Public/Dom/JsonObject.h b/Engine/Source/Runtime/Json/Public/Dom/JsonObject.h
--- a/Engine/Source/Runtime/Json/Public/Dom/JsonObject.h
+++ b/Engine/Source/Runtime/Json/Public/Dom/JsonObject.h
@@ -26,10 +26,15 @@ public:
 	TMap<FString, TSharedPtr<FJsonValue>> Values;
 
 	template<EJson JsonType>
-	TSharedPtr<FJsonValue> GetField( const FString& FieldName ) const
+	TSharedPtr<FJsonValue> GetField(FStringView FieldName) const
 	{
-		const TSharedPtr<FJsonValue>* Field = Values.Find(FieldName);
-		if ( Field != nullptr && Field->IsValid() )
+		return GetField(FieldName, JsonType);
+	}
+
+	TSharedPtr<FJsonValue> GetField(FStringView FieldName, EJson JsonType) const
+	{
+		const TSharedPtr<FJsonValue>* Field = Values.FindByHash(GetTypeHash(FieldName), FieldName);
+		if (Field != nullptr && Field->IsValid())
 		{
 			if (JsonType == EJson::None || (*Field)->Type == JsonType)
 			{
@@ -37,12 +42,12 @@ public:
 			}
 			else
 			{
-				UE_LOG(LogJson, Warning, TEXT("Field %s is of the wrong type."), *FieldName);
+				UE_LOG(LogJson, Warning, TEXT("Field %.*s is of the wrong type."), FieldName.Len(), FieldName.GetData());
 			}
 		}
 		else
 		{
-			UE_LOG(LogJson, Warning, TEXT("Field %s was not found."), *FieldName);
+			UE_LOG(LogJson, Warning, TEXT("Field %.*s was not found."), FieldName.Len(), FieldName.GetData());
 		}
 
 		return MakeShared<FJsonValueNull>();
@@ -54,9 +59,9 @@ public:
 	 * @param FieldName The name of the field to get.
 	 * @return A pointer to the field, or nullptr if the field doesn't exist.
 	 */
-	TSharedPtr<FJsonValue> TryGetField( const FString& FieldName ) const
+	TSharedPtr<FJsonValue> TryGetField(FStringView FieldName) const
 	{
-		const TSharedPtr<FJsonValue>* Field = Values.Find(FieldName);
+		const TSharedPtr<FJsonValue>* Field = Values.FindByHash(GetTypeHash(FieldName), FieldName);
 		return (Field != nullptr && Field->IsValid()) ? *Field : TSharedPtr<FJsonValue>();
 	}
 
@@ -66,9 +71,9 @@ public:
 	 * @param FieldName The name of the field to check.
 	 * @return true if the field exists, false otherwise.
 	 */
-	bool HasField( const FString& FieldName) const
+	bool HasField(FStringView FieldName) const
 	{
-		const TSharedPtr<FJsonValue>* Field = Values.Find(FieldName);
+		const TSharedPtr<FJsonValue>* Field = Values.FindByHash(GetTypeHash(FieldName), FieldName);
 		if(Field && Field->IsValid())
 		{
 			return true;
@@ -80,15 +85,27 @@ public:
 	/**
 	 * Checks whether a field with the specified name and type exists in the object.
 	 *
-	 * @param JsonType The type of the field to check.
+	 * @tparam JsonType The type of the field to check.
 	 * @param FieldName The name of the field to check.
 	 * @return true if the field exists, false otherwise.
 	 */
 	template<EJson JsonType>
-	bool HasTypedField(const FString& FieldName) const
+	bool HasTypedField(FStringView FieldName) const
 	{
-		const TSharedPtr<FJsonValue>* Field = Values.Find(FieldName);
-		if(Field && Field->IsValid() && ((*Field)->Type == JsonType))
+		return HasTypedField(FieldName, JsonType);
+	}
+
+	/**
+	 * Checks whether a field with the specified name and type exists in the object.
+	 *
+	 * @param JsonType The type of the field to check.
+	 * @param FieldName The name of the field to check.
+	 * @return true if the field exists, false otherwise.
+	 */
+	bool HasTypedField(FStringView FieldName, EJson JsonType) const
+	{
+		const TSharedPtr<FJsonValue>* Field = Values.FindByHash(GetTypeHash(FieldName), FieldName);
+		if (Field && Field->IsValid() && ((*Field)->Type == JsonType))
 		{
 			return true;
 		}
@@ -102,14 +119,15 @@ public:
 	 * @param FieldName The name of the field to set.
 	 * @param Value The value to set.
 	 */
-	JSON_API void SetField( const FString& FieldName, const TSharedPtr<FJsonValue>& Value );
+	JSON_API void SetField(FString&& FieldName, const TSharedPtr<FJsonValue>& Value);
+	JSON_API void SetField(const FString& FieldName, const TSharedPtr<FJsonValue>& Value);
 
 	/**
 	 * Removes the field with the specified name.
 	 *
 	 * @param FieldName The name of the field to remove.
 	 */
-	JSON_API void RemoveField(const FString& FieldName);
+	JSON_API void RemoveField(FStringView FieldName);
 
 	/**
 	 * Gets the field with the specified name as a number.
@@ -119,61 +137,62 @@ public:
 	 * @param FieldName The name of the field to get.
 	 * @return The field's value as a number.
 	 */
-	JSON_API double GetNumberField(const FString& FieldName) const;
+	JSON_API double GetNumberField(FStringView FieldName) const;
 
 	/**
 	 * Gets a numeric field and casts to an int32
 	 */
-	FORCEINLINE int32 GetIntegerField(const FString& FieldName) const
+	FORCEINLINE int32 GetIntegerField(FStringView FieldName) const
 	{
 		return (int32)GetNumberField(FieldName);
 	}
 
 	/** Get the field named FieldName as a number. Returns false if it doesn't exist or cannot be converted. */
-	JSON_API bool TryGetNumberField(const FString& FieldName, float& OutNumber) const;
+	JSON_API bool TryGetNumberField(FStringView FieldName, float& OutNumber) const;
 
 	/** Get the field named FieldName as a number. Returns false if it doesn't exist or cannot be converted. */
-	JSON_API bool TryGetNumberField(const FString& FieldName, double& OutNumber) const;
+	JSON_API bool TryGetNumberField(FStringView FieldName, double& OutNumber) const;
 
 	/** Get the field named FieldName as a number, and makes sure it's within int8 range. Returns false if it doesn't exist or cannot be converted. */
-	JSON_API bool TryGetNumberField(const FString& FieldName, int8& OutNumber) const;
+	JSON_API bool TryGetNumberField(FStringView FieldName, int8& OutNumber) const;
 
 	/** Get the field named FieldName as a number, and makes sure it's within int16 range. Returns false if it doesn't exist or cannot be converted. */
-	JSON_API bool TryGetNumberField(const FString& FieldName, int16& OutNumber) const;
+	JSON_API bool TryGetNumberField(FStringView FieldName, int16& OutNumber) const;
 
 	/** Get the field named FieldName as a number, and makes sure it's within int32 range. Returns false if it doesn't exist or cannot be converted. */
-	JSON_API bool TryGetNumberField(const FString& FieldName, int32& OutNumber) const;
+	JSON_API bool TryGetNumberField(FStringView FieldName, int32& OutNumber) const;
 
 	/** Get the field named FieldName as a number. Returns false if it doesn't exist or cannot be converted. */
-	JSON_API bool TryGetNumberField(const FString& FieldName, int64& OutNumber) const;
+	JSON_API bool TryGetNumberField(FStringView FieldName, int64& OutNumber) const;
 
 	/** Get the field named FieldName as a number, and makes sure it's within uint8 range. Returns false if it doesn't exist or cannot be converted. */
-	JSON_API bool TryGetNumberField(const FString& FieldName, uint8& OutNumber) const;
+	JSON_API bool TryGetNumberField(FStringView FieldName, uint8& OutNumber) const;
 
 	/** Get the field named FieldName as a number, and makes sure it's within uint16 range. Returns false if it doesn't exist or cannot be converted. */
-	JSON_API bool TryGetNumberField(const FString& FieldName, uint16& OutNumber) const;
+	JSON_API bool TryGetNumberField(FStringView FieldName, uint16& OutNumber) const;
 
 	/** Get the field named FieldName as a number, and makes sure it's within uint32 range. Returns false if it doesn't exist or cannot be converted.  */
-	JSON_API bool TryGetNumberField(const FString& FieldName, uint32& OutNumber) const;
+	JSON_API bool TryGetNumberField(FStringView FieldName, uint32& OutNumber) const;
 
 	/** Get the field named FieldName as a number, and makes sure it's within uint64 range. Returns false if it doesn't exist or cannot be converted.  */
-	JSON_API bool TryGetNumberField(const FString& FieldName, uint64& OutNumber) const;
+	JSON_API bool TryGetNumberField(FStringView FieldName, uint64& OutNumber) const;
 
 	/** Add a field named FieldName with Number as value */
-	JSON_API void SetNumberField( const FString& FieldName, double Number );
+	JSON_API void SetNumberField(FString&& FieldName, double Number);
+	JSON_API void SetNumberField(const FString& FieldName, double Number);
 
 	/** Get the field named FieldName as a string. */
-	JSON_API FString GetStringField(const FString& FieldName) const;
+	JSON_API FString GetStringField(FStringView FieldName) const;
 
 	/** Get the field named FieldName as a string. Returns false if it doesn't exist or cannot be converted. */
-	JSON_API bool TryGetStringField(const FString& FieldName, FString& OutString) const;
+	JSON_API bool TryGetStringField(FStringView FieldName, FString& OutString) const;
 
 	/** Get the field named FieldName as an array of strings. Returns false if it doesn't exist or any member cannot be converted. */
-	JSON_API bool TryGetStringArrayField(const FString& FieldName, TArray<FString>& OutArray) const;
+	JSON_API bool TryGetStringArrayField(FStringView FieldName, TArray<FString>& OutArray) const;
 
 	/** Get the field named FieldName as an array of enums. Returns false if it doesn't exist or any member is not a string. */
 	template<typename TEnum>
-	bool TryGetEnumArrayField(const FString& FieldName, TArray<TEnum>& OutArray) const
+	bool TryGetEnumArrayField(FStringView FieldName, TArray<TEnum>& OutArray) const
 	{
 		TArray<FString> Strings;
 		if (!TryGetStringArrayField(FieldName, Strings))
@@ -194,7 +213,10 @@ public:
 	}
 
 	/** Add a field named FieldName with value of StringValue */
-	JSON_API void SetStringField( const FString& FieldName, const FString& StringValue );
+	JSON_API void SetStringField(FString&& FieldName, FString&& StringValue);
+	JSON_API void SetStringField(FString&& FieldName, const FString& StringValue);
+	JSON_API void SetStringField(const FString& FieldName, FString&& StringValue);
+	JSON_API void SetStringField(const FString& FieldName, const FString& StringValue);
 
 	/**
 	 * Gets the field with the specified name as a boolean.
@@ -204,22 +226,26 @@ public:
 	 * @param FieldName The name of the field to get.
 	 * @return The field's value as a boolean.
 	 */
-	JSON_API bool GetBoolField(const FString& FieldName) const;
+	JSON_API bool GetBoolField(FStringView FieldName) const;
 
 	/** Get the field named FieldName as a string. Returns false if it doesn't exist or cannot be converted. */
-	JSON_API bool TryGetBoolField(const FString& FieldName, bool& OutBool) const;
+	JSON_API bool TryGetBoolField(FStringView FieldName, bool& OutBool) const;
 
 	/** Set a boolean field named FieldName and value of InValue */
-	JSON_API void SetBoolField( const FString& FieldName, bool InValue );
+	JSON_API void SetBoolField(FString&& FieldName, bool InValue);
+	JSON_API void SetBoolField(const FString& FieldName, bool InValue);
 
 	/** Get the field named FieldName as an array. */
-	JSON_API const TArray< TSharedPtr<FJsonValue> >& GetArrayField(const FString& FieldName) const;
+	JSON_API const TArray<TSharedPtr<FJsonValue>>& GetArrayField(FStringView FieldName) const;
 
 	/** Try to get the field named FieldName as an array, or return false if it's another type */
-	JSON_API bool TryGetArrayField(const FString& FieldName, const TArray< TSharedPtr<FJsonValue> >*& OutArray) const;
+	JSON_API bool TryGetArrayField(FStringView FieldName, const TArray<TSharedPtr<FJsonValue>>*& OutArray) const;
 
 	/** Set an array field named FieldName and value of Array */
-	JSON_API void SetArrayField( const FString& FieldName, const TArray< TSharedPtr<FJsonValue> >& Array );
+	JSON_API void SetArrayField(FString&& FieldName, TArray<TSharedPtr<FJsonValue>>&& Array);
+	JSON_API void SetArrayField(FString&& FieldName, const TArray<TSharedPtr<FJsonValue>>& Array);
+	JSON_API void SetArrayField(const FString& FieldName, TArray<TSharedPtr<FJsonValue>>&& Array);
+	JSON_API void SetArrayField(const FString& FieldName, const TArray<TSharedPtr<FJsonValue>>& Array);
 
 	/**
 	 * Gets the field with the specified name as a Json object.
@@ -229,14 +255,185 @@ public:
 	 * @param FieldName The name of the field to get.
 	 * @return The field's value as a Json object.
 	 */
-	JSON_API const TSharedPtr<FJsonObject>& GetObjectField(const FString& FieldName) const;
+	JSON_API const TSharedPtr<FJsonObject>& GetObjectField(FStringView FieldName) const;
 
 	/** Try to get the field named FieldName as an object, or return false if it's another type */
-	JSON_API bool TryGetObjectField(const FString& FieldName, const TSharedPtr<FJsonObject>*& OutObject) const;
+	JSON_API bool TryGetObjectField(FStringView FieldName, const TSharedPtr<FJsonObject>*& OutObject) const;
 
 	/** Set an ObjectField named FieldName and value of JsonObject */
-	JSON_API void SetObjectField( const FString& FieldName, const TSharedPtr<FJsonObject>& JsonObject );
+	JSON_API void SetObjectField(FString&& FieldName, const TSharedPtr<FJsonObject>& JsonObject);
+	JSON_API void SetObjectField(const FString& FieldName, const TSharedPtr<FJsonObject>& JsonObject);
 
 	static JSON_API void Duplicate(const TSharedPtr<const FJsonObject>& Source, const TSharedPtr<FJsonObject>& Dest);
 	static JSON_API void Duplicate(const TSharedPtr<FJsonObject>& Source, TSharedPtr<FJsonObject>& Dest);
+
+#if !PLATFORM_TCHAR_IS_UTF8CHAR
+
+	template <EJson JsonType>
+	UE_DEPRECATED(5.4, "Passing an ANSI string to GetField has been deprecated outside of UTF-8 mode. Please use the overload that takes a TCHAR string.")
+	TSharedPtr<FJsonValue> GetField(FAnsiStringView FieldName) const
+	{
+		return GetField<JsonType>(StringCast<TCHAR>(FieldName.GetData(), FieldName.Len()));
+	}
+
+	UE_DEPRECATED(5.4, "Passing an ANSI string to TryGetField has been deprecated outside of UTF-8 mode. Please use the overload that takes a TCHAR string.")
+	TSharedPtr<FJsonValue> TryGetField(FAnsiStringView FieldName) const
+	{
+		return TryGetField(StringCast<TCHAR>(FieldName.GetData(), FieldName.Len()));
+	}
+
+	UE_DEPRECATED(5.4, "Passing an ANSI string to HasField has been deprecated outside of UTF-8 mode. Please use the overload that takes a TCHAR string.")
+	bool HasField(FAnsiStringView FieldName) const
+	{
+		return HasField(StringCast<TCHAR>(FieldName.GetData(), FieldName.Len()));
+	}
+
+	template <EJson JsonType>
+	UE_DEPRECATED(5.4, "Passing an ANSI string to HasTypedField has been deprecated outside of UTF-8 mode. Please use the overload that takes a TCHAR string.")
+	bool HasTypedField(FAnsiStringView FieldName) const
+	{
+		return HasTypedField<JsonType>(StringCast<TCHAR>(FieldName.GetData(), FieldName.Len()));
+	}
+
+	UE_DEPRECATED(5.4, "Passing an ANSI string to RemoveField has been deprecated outside of UTF-8 mode. Please use the overload that takes a TCHAR string.")
+	void RemoveField(FAnsiStringView FieldName)
+	{
+		return RemoveField(StringCast<TCHAR>(FieldName.GetData(), FieldName.Len()));
+	}
+
+	UE_DEPRECATED(5.4, "Passing an ANSI string to GetNumberField has been deprecated outside of UTF-8 mode. Please use the overload that takes a TCHAR string.")
+	double GetNumberField(FAnsiStringView FieldName) const
+	{
+		return GetNumberField(StringCast<TCHAR>(FieldName.GetData(), FieldName.Len()));
+	}
+
+	UE_DEPRECATED(5.4, "Passing an ANSI string to GetIntegerField has been deprecated outside of UTF-8 mode. Please use the overload that takes a TCHAR string.")
+	int32 GetIntegerField(FAnsiStringView FieldName) const
+	{
+		return GetIntegerField(StringCast<TCHAR>(FieldName.GetData(), FieldName.Len()));
+	}
+
+	UE_DEPRECATED(5.4, "Passing an ANSI string to TryGetNumberField has been deprecated outside of UTF-8 mode. Please use the overload that takes a TCHAR string.")
+	bool TryGetNumberField(FAnsiStringView FieldName, float& OutNumber) const
+	{
+		return TryGetNumberField(StringCast<TCHAR>(FieldName.GetData(), FieldName.Len()), OutNumber);
+	}
+
+	UE_DEPRECATED(5.4, "Passing an ANSI string to TryGetNumberField has been deprecated outside of UTF-8 mode. Please use the overload that takes a TCHAR string.")
+	bool TryGetNumberField(FAnsiStringView FieldName, double& OutNumber) const
+	{
+		return TryGetNumberField(StringCast<TCHAR>(FieldName.GetData(), FieldName.Len()), OutNumber);
+	}
+
+	UE_DEPRECATED(5.4, "Passing an ANSI string to TryGetNumberField has been deprecated outside of UTF-8 mode. Please use the overload that takes a TCHAR string.")
+	bool TryGetNumberField(FAnsiStringView FieldName, int8& OutNumber) const
+	{
+		return TryGetNumberField(StringCast<TCHAR>(FieldName.GetData(), FieldName.Len()), OutNumber);
+	}
+
+	UE_DEPRECATED(5.4, "Passing an ANSI string to TryGetNumberField has been deprecated outside of UTF-8 mode. Please use the overload that takes a TCHAR string.")
+	bool TryGetNumberField(FAnsiStringView FieldName, int16& OutNumber) const
+	{
+		return TryGetNumberField(StringCast<TCHAR>(FieldName.GetData(), FieldName.Len()), OutNumber);
+	}
+
+	UE_DEPRECATED(5.4, "Passing an ANSI string to TryGetNumberField has been deprecated outside of UTF-8 mode. Please use the overload that takes a TCHAR string.")
+	bool TryGetNumberField(FAnsiStringView FieldName, int32& OutNumber) const
+	{
+		return TryGetNumberField(StringCast<TCHAR>(FieldName.GetData(), FieldName.Len()), OutNumber);
+	}
+
+	UE_DEPRECATED(5.4, "Passing an ANSI string to TryGetNumberField has been deprecated outside of UTF-8 mode. Please use the overload that takes a TCHAR string.")
+	bool TryGetNumberField(FAnsiStringView FieldName, int64& OutNumber) const
+	{
+		return TryGetNumberField(StringCast<TCHAR>(FieldName.GetData(), FieldName.Len()), OutNumber);
+	}
+
+	UE_DEPRECATED(5.4, "Passing an ANSI string to TryGetNumberField has been deprecated outside of UTF-8 mode. Please use the overload that takes a TCHAR string.")
+	bool TryGetNumberField(FAnsiStringView FieldName, uint8& OutNumber) const
+	{
+		return TryGetNumberField(StringCast<TCHAR>(FieldName.GetData(), FieldName.Len()), OutNumber);
+	}
+
+	UE_DEPRECATED(5.4, "Passing an ANSI string to TryGetNumberField has been deprecated outside of UTF-8 mode. Please use the overload that takes a TCHAR string.")
+	bool TryGetNumberField(FAnsiStringView FieldName, uint16& OutNumber) const
+	{
+		return TryGetNumberField(StringCast<TCHAR>(FieldName.GetData(), FieldName.Len()), OutNumber);
+	}
+
+	UE_DEPRECATED(5.4, "Passing an ANSI string to TryGetNumberField has been deprecated outside of UTF-8 mode. Please use the overload that takes a TCHAR string.")
+	bool TryGetNumberField(FAnsiStringView FieldName, uint32& OutNumber) const
+	{
+		return TryGetNumberField(StringCast<TCHAR>(FieldName.GetData(), FieldName.Len()), OutNumber);
+	}
+
+	UE_DEPRECATED(5.4, "Passing an ANSI string to TryGetNumberField has been deprecated outside of UTF-8 mode. Please use the overload that takes a TCHAR string.")
+	bool TryGetNumberField(FAnsiStringView FieldName, uint64& OutNumber) const
+	{
+		return TryGetNumberField(StringCast<TCHAR>(FieldName.GetData(), FieldName.Len()), OutNumber);
+	}
+
+	UE_DEPRECATED(5.4, "Passing an ANSI string to GetStringField has been deprecated outside of UTF-8 mode. Please use the overload that takes a TCHAR string.")
+	FString GetStringField(FAnsiStringView FieldName) const
+	{
+		return GetStringField(StringCast<TCHAR>(FieldName.GetData(), FieldName.Len()));
+	}
+
+	UE_DEPRECATED(5.4, "Passing an ANSI string to TryGetStringField has been deprecated outside of UTF-8 mode. Please use the overload that takes a TCHAR string.")
+	bool TryGetStringField(FAnsiStringView FieldName, FString& OutString) const
+	{
+		return TryGetStringField(StringCast<TCHAR>(FieldName.GetData(), FieldName.Len()), OutString);
+	}
+
+	UE_DEPRECATED(5.4, "Passing an ANSI string to TryGetStringArrayField has been deprecated outside of UTF-8 mode. Please use the overload that takes a TCHAR string.")
+	bool TryGetStringArrayField(FAnsiStringView FieldName, TArray<FString>& OutArray) const
+	{
+		return TryGetStringArrayField(StringCast<TCHAR>(FieldName.GetData(), FieldName.Len()), OutArray);
+	}
+
+	template <typename TEnum>
+	UE_DEPRECATED(5.4, "Passing an ANSI string to TryGetEnumArrayField has been deprecated outside of UTF-8 mode. Please use the overload that takes a TCHAR string.")
+	bool TryGetEnumArrayField(FAnsiStringView FieldName, TArray<TEnum>& OutArray) const
+	{
+		return TryGetEnumArrayField(StringCast<TCHAR>(FieldName.GetData(), FieldName.Len()), OutArray);
+	}
+
+	UE_DEPRECATED(5.4, "Passing an ANSI string to GetBoolField has been deprecated outside of UTF-8 mode. Please use the overload that takes a TCHAR string.")
+	bool GetBoolField(FAnsiStringView FieldName) const
+	{
+		return GetBoolField(StringCast<TCHAR>(FieldName.GetData(), FieldName.Len()));
+	}
+
+	UE_DEPRECATED(5.4, "Passing an ANSI string to TryGetBoolField has been deprecated outside of UTF-8 mode. Please use the overload that takes a TCHAR string.")
+	bool TryGetBoolField(FAnsiStringView FieldName, bool& OutBool) const
+	{
+		return TryGetBoolField(StringCast<TCHAR>(FieldName.GetData(), FieldName.Len()), OutBool);
+	}
+
+	UE_DEPRECATED(5.4, "Passing an ANSI string to GetArrayField has been deprecated outside of UTF-8 mode. Please use the overload that takes a TCHAR string.")
+	const TArray<TSharedPtr<FJsonValue>>& GetArrayField(FAnsiStringView FieldName) const
+	{
+		return GetArrayField(StringCast<TCHAR>(FieldName.GetData(), FieldName.Len()));
+	}
+
+	UE_DEPRECATED(5.4, "Passing an ANSI string to TryGetArrayField has been deprecated outside of UTF-8 mode. Please use the overload that takes a TCHAR string.")
+	bool TryGetArrayField(FAnsiStringView FieldName, const TArray<TSharedPtr<FJsonValue>>*& OutArray) const
+	{
+		return TryGetArrayField(StringCast<TCHAR>(FieldName.GetData(), FieldName.Len()), OutArray);
+	}
+
+	UE_DEPRECATED(5.4, "Passing an ANSI string to GetObjectField has been deprecated outside of UTF-8 mode. Please use the overload that takes a TCHAR string.")
+	const TSharedPtr<FJsonObject>& GetObjectField(FAnsiStringView FieldName) const
+	{
+		return GetObjectField(StringCast<TCHAR>(FieldName.GetData(), FieldName.Len()));
+	}
+
+	UE_DEPRECATED(5.4, "Passing an ANSI string to TryGetObjectField has been deprecated outside of UTF-8 mode. Please use the overload that takes a TCHAR string.")
+	bool TryGetObjectField(FAnsiStringView FieldName, const TSharedPtr<FJsonObject>*& OutObject) const
+	{
+		return TryGetObjectField(StringCast<TCHAR>(FieldName.GetData(), FieldName.Len()), OutObject);
+	}
+
+#endif // !PLATFORM_TCHAR_IS_UTF8CHAR
+
 };
diff --git a/Engine/Source/Runtime/Json/Public/Serialization/JsonSerializerMacros.h b/Engine/Source/Runtime/Json/Public/Serialization/JsonSerializerMacros.h
--- a/Engine/Source/Runtime/Json/Public/Serialization/JsonSerializerMacros.h
+++ b/Engine/Source/Runtime/Json/Public/Serialization/JsonSerializerMacros.h
@@ -22,65 +22,65 @@
 	}
 
 #define JSON_SERIALIZE(JsonName, JsonValue) \
-		Serializer.Serialize(TEXT(JsonName), JsonValue)
+		Serializer.Serialize(TEXTVIEW(JsonName), JsonValue)
 
 #define JSON_SERIALIZE_WITHDEFAULT(JsonName, JsonValue, DefaultJsonValue) \
 		if (Serializer.IsLoading()) \
 		{ \
-			if (!Serializer.GetObject()->HasField(TEXT(JsonName))) \
+			if (!Serializer.GetObject()->HasField(TEXTVIEW(JsonName))) \
            	{ \
            		JsonValue = DefaultJsonValue; \
            	} \
         } \
-		Serializer.Serialize(TEXT(JsonName), JsonValue);
+		Serializer.Serialize(TEXTVIEW(JsonName), JsonValue);
 		
 #define JSON_SERIALIZE_OPTIONAL(JsonName, OptionalJsonValue) \
 		if (Serializer.IsLoading()) \
 		{ \
-			if (Serializer.GetObject()->HasField(TEXT(JsonName))) \
+			if (Serializer.GetObject()->HasField(TEXTVIEW(JsonName))) \
 			{ \
-				Serializer.Serialize(TEXT(JsonName), OptionalJsonValue.Emplace()); \
+				Serializer.Serialize(TEXTVIEW(JsonName), OptionalJsonValue.Emplace()); \
 			} \
 		} \
 		else \
 		{ \
 			if (OptionalJsonValue.IsSet()) \
 			{ \
-				Serializer.Serialize(TEXT(JsonName), OptionalJsonValue.GetValue()); \
+				Serializer.Serialize(TEXTVIEW(JsonName), OptionalJsonValue.GetValue()); \
 			} \
 		}
 
 #define JSON_SERIALIZE_ARRAY(JsonName, JsonArray) \
-		Serializer.SerializeArray(TEXT(JsonName), JsonArray)
+		Serializer.SerializeArray(TEXTVIEW(JsonName), JsonArray)
 		
 #define JSON_SERIALIZE_ARRAY_WITHDEFAULT(JsonName, JsonArray, DefaultArray) \
 		if (Serializer.IsLoading()) \
 		{ \
-			if (!Serializer.GetObject()->HasField(TEXT(JsonName))) \
+			if (!Serializer.GetObject()->HasField(TEXTVIEW(JsonName))) \
 			{ \
 				JsonArray = DefaultArray; \
 			} \
 		} \
-		Serializer.SerializeArray(TEXT(JsonName), JsonArray);
+		Serializer.SerializeArray(TEXTVIEW(JsonName), JsonArray);
 		
 #define JSON_SERIALIZE_OBJECT_WITHDEFAULT(JsonName, JsonObject, DefaultObject) \
 		if (Serializer.IsLoading()) \
 		{ \
-			if (!Serializer.GetObject()->HasField(TEXT(JsonName))) \
+			if (!Serializer.GetObject()->HasField(TEXTVIEW(JsonName))) \
 			{ \
 				JsonObject = DefaultObject; \
 			} \
 		} \
-		Serializer.SerializeArray(TEXT(JsonName), JsonArray);
+		Serializer.SerializeArray(TEXTVIEW(JsonName), JsonArray);
 
 #define JSON_SERIALIZE_MAP(JsonName, JsonMap) \
-		Serializer.SerializeMap(TEXT(JsonName), JsonMap)
+		Serializer.SerializeMap(TEXTVIEW(JsonName), JsonMap)
 
 #define JSON_SERIALIZE_SIMPLECOPY(JsonMap) \
 		Serializer.SerializeSimpleMap(JsonMap)
 
 #define JSON_SERIALIZE_MAP_SAFE(JsonName, JsonMap) \
-		Serializer.SerializeMapSafe(TEXT(JsonName), JsonMap)
+		Serializer.SerializeMapSafe(TEXTVIEW(JsonName), JsonMap)
 
 #define JSON_SERIALIZE_SERIALIZABLE(JsonName, JsonValue) \
 		JsonValue.Serialize(Serializer, false)
@@ -88,9 +88,9 @@
 #define JSON_SERIALIZE_RAW_JSON_STRING(JsonName, JsonValue) \
 		if (Serializer.IsLoading()) \
 		{ \
-			if (Serializer.GetObject()->HasTypedField<EJson::Object>(TEXT(JsonName))) \
+			if (Serializer.GetObject()->HasTypedField<EJson::Object>(TEXTVIEW(JsonName))) \
 			{ \
-				TSharedPtr<FJsonObject> JsonObject = Serializer.GetObject()->GetObjectField(TEXT(JsonName)); \
+				TSharedPtr<FJsonObject> JsonObject = Serializer.GetObject()->GetObjectField(TEXTVIEW(JsonName)); \
 				if (JsonObject.IsValid()) \
 				{ \
 					auto Writer = TJsonWriterFactory<TCHAR, TCondensedJsonPrintPolicy<TCHAR>>::Create(&JsonValue); \
@@ -106,7 +106,7 @@
 		{ \
 			if (!JsonValue.IsEmpty()) \
 			{ \
-				Serializer.WriteIdentifierPrefix(TEXT(JsonName)); \
+				Serializer.WriteIdentifierPrefix(TEXTVIEW(JsonName)); \
 				Serializer.WriteRawJSONValue(*JsonValue); \
 			} \
 		}
@@ -114,9 +114,9 @@
 #define JSON_SERIALIZE_ARRAY_SERIALIZABLE(JsonName, JsonArray, ElementType) \
 		if (Serializer.IsLoading()) \
 		{ \
-			if (Serializer.GetObject()->HasTypedField<EJson::Array>(JsonName)) \
+			if (Serializer.GetObject()->HasTypedField<EJson::Array>(TEXTVIEW(JsonName))) \
 			{ \
-				for (auto It = Serializer.GetObject()->GetArrayField(JsonName).CreateConstIterator(); It; ++It) \
+				for (auto It = Serializer.GetObject()->GetArrayField(TEXTVIEW(JsonName)).CreateConstIterator(); It; ++It) \
 				{ \
 					ElementType& Obj = JsonArray.AddDefaulted_GetRef(); \
 					Obj.FromJson((*It)->AsObject()); \
@@ -125,7 +125,7 @@
 		} \
 		else \
 		{ \
-			Serializer.StartArray(JsonName); \
+			Serializer.StartArray(TEXTVIEW(JsonName)); \
 			for (auto It = JsonArray.CreateIterator(); It; ++It) \
 			{ \
 				It->Serialize(Serializer, false); \
@@ -136,9 +136,9 @@
 #define JSON_SERIALIZE_ARRAY_SERIALIZABLE_WITHDEFAULT(JsonName, JsonArray, ElementType, DefaultArray) \
 		if (Serializer.IsLoading()) \
 		{ \
-			if (Serializer.GetObject()->HasTypedField<EJson::Array>(JsonName)) \
+			if (Serializer.GetObject()->HasTypedField<EJson::Array>(TEXTVIEW(JsonName))) \
 			{ \
-				for (auto It = Serializer.GetObject()->GetArrayField(JsonName).CreateConstIterator(); It; ++It) \
+				for (auto It = Serializer.GetObject()->GetArrayField(TEXTVIEW(JsonName)).CreateConstIterator(); It; ++It) \
 				{ \
 					ElementType& Obj = JsonArray.AddDefaulted_GetRef(); \
 					Obj.FromJson((*It)->AsObject()); \
@@ -151,7 +151,7 @@
 		} \
 		else \
 		{ \
-			Serializer.StartArray(JsonName); \
+			Serializer.StartArray(TEXTVIEW(JsonName)); \
 			for (auto It = JsonArray.CreateIterator(); It; ++It) \
 			{ \
 				It->Serialize(Serializer, false); \
@@ -162,10 +162,10 @@
 #define JSON_SERIALIZE_OPTIONAL_ARRAY_SERIALIZABLE(JsonName, OptionalJsonArray, ElementType) \
 		if (Serializer.IsLoading()) \
 		{ \
-			if (Serializer.GetObject()->HasTypedField<EJson::Array>(JsonName)) \
+			if (Serializer.GetObject()->HasTypedField<EJson::Array>(TEXTVIEW(JsonName))) \
 			{ \
 				TArray<ElementType>& JsonArray = OptionalJsonArray.Emplace(); \
-				for (auto It = Serializer.GetObject()->GetArrayField(JsonName).CreateConstIterator(); It; ++It) \
+				for (auto It = Serializer.GetObject()->GetArrayField(TEXTVIEW(JsonName)).CreateConstIterator(); It; ++It) \
 				{ \
 					ElementType& Obj = JsonArray.AddDefaulted_GetRef(); \
 					Obj.FromJson((*It)->AsObject()); \
@@ -176,7 +176,7 @@
 		{ \
 			if (OptionalJsonArray.IsSet()) \
 			{ \
-				Serializer.StartArray(JsonName); \
+				Serializer.StartArray(TEXTVIEW(JsonName)); \
 				for (auto It = OptionalJsonArray->CreateIterator(); It; ++It) \
 				{ \
 					It->Serialize(Serializer, false); \
@@ -188,9 +188,9 @@
 #define JSON_SERIALIZE_MAP_SERIALIZABLE(JsonName, JsonMap, ElementType) \
 		if (Serializer.IsLoading()) \
 		{ \
-			if (Serializer.GetObject()->HasTypedField<EJson::Object>(JsonName)) \
+			if (Serializer.GetObject()->HasTypedField<EJson::Object>(TEXTVIEW(JsonName))) \
 			{ \
-				TSharedPtr<FJsonObject> JsonObj = Serializer.GetObject()->GetObjectField(JsonName); \
+				TSharedPtr<FJsonObject> JsonObj = Serializer.GetObject()->GetObjectField(TEXTVIEW(JsonName)); \
 				for (auto MapIt = JsonObj->Values.CreateConstIterator(); MapIt; ++MapIt) \
 				{ \
 					ElementType NewEntry; \
@@ -201,10 +201,10 @@
 		} \
 		else \
 		{ \
-			Serializer.StartObject(JsonName); \
+			Serializer.StartObject(TEXTVIEW(JsonName)); \
 			for (auto It = JsonMap.CreateIterator(); It; ++It) \
 			{ \
-				Serializer.StartObject(It.Key()); \
+				Serializer.StartObject(MakeStringView(It.Key())); \
 				It.Value().Serialize(Serializer, true); \
 				Serializer.EndObject(); \
 			} \
@@ -216,9 +216,9 @@
 		if (Serializer.IsLoading()) \
 		{ \
 			/* Read in the value from the JsonName field */ \
-			if (Serializer.GetObject()->HasTypedField<EJson::Object>(JsonName)) \
+			if (Serializer.GetObject()->HasTypedField<EJson::Object>(TEXTVIEW(JsonName))) \
 			{ \
-				TSharedPtr<FJsonObject> JsonObj = Serializer.GetObject()->GetObjectField(JsonName); \
+				TSharedPtr<FJsonObject> JsonObj = Serializer.GetObject()->GetObjectField(TEXTVIEW(JsonName)); \
 				if (JsonObj.IsValid()) \
 				{ \
 					(JsonSerializableObject).FromJson(JsonObj); \
@@ -228,7 +228,7 @@
 		else \
 		{ \
 			/* Write the value to the Name field */ \
-			Serializer.StartObject(JsonName); \
+			Serializer.StartObject(TEXTVIEW(JsonName)); \
 			(JsonSerializableObject).Serialize(Serializer, true); \
 			Serializer.EndObject(); \
 		}
@@ -238,9 +238,9 @@
 		if (Serializer.IsLoading()) \
 		{ \
 			/* Read in the value from the JsonName field */ \
-			if (Serializer.GetObject()->HasTypedField<EJson::Object>(JsonName)) \
+			if (Serializer.GetObject()->HasTypedField<EJson::Object>(TEXTVIEW(JsonName))) \
 			{ \
-				TSharedPtr<FJsonObject> JsonObj = Serializer.GetObject()->GetObjectField(JsonName); \
+				TSharedPtr<FJsonObject> JsonObj = Serializer.GetObject()->GetObjectField(TEXTVIEW(JsonName)); \
 				if (JsonObj.IsValid()) \
 				{ \
 					(JsonSerializableObject).FromJson(JsonObj); \
@@ -254,7 +254,7 @@
 		else \
 		{ \
 			/* Write the value to the Name field */ \
-			Serializer.StartObject(JsonName); \
+			Serializer.StartObject(TEXTVIEW(JsonName)); \
 			(JsonSerializableObject).Serialize(Serializer, true); \
 			Serializer.EndObject(); \
 		}
@@ -263,9 +263,9 @@
 		if (Serializer.IsLoading()) \
 		{ \
 			using ObjectType = TRemoveReference<decltype(JsonSerializableObject.GetValue())>::Type; \
-			if (Serializer.GetObject()->HasTypedField<EJson::Object>(JsonName)) \
+			if (Serializer.GetObject()->HasTypedField<EJson::Object>(TEXTVIEW(JsonName))) \
 			{ \
-				TSharedPtr<FJsonObject> JsonObj = Serializer.GetObject()->GetObjectField(JsonName); \
+				TSharedPtr<FJsonObject> JsonObj = Serializer.GetObject()->GetObjectField(TEXTVIEW(JsonName)); \
 				if (JsonObj.IsValid()) \
 				{ \
 					JsonSerializableObject = ObjectType{}; \
@@ -277,7 +277,7 @@
 		{ \
 			if (JsonSerializableObject.IsSet()) \
 			{ \
-				Serializer.StartObject(JsonName); \
+				Serializer.StartObject(TEXTVIEW(JsonName)); \
 				(JsonSerializableObject.GetValue()).Serialize(Serializer, true); \
 				Serializer.EndObject(); \
 			} \
@@ -287,48 +287,48 @@
 		if (Serializer.IsLoading()) \
 		{ \
 			int64 UnixTimestampValue; \
-			Serializer.Serialize(TEXT(JsonName), UnixTimestampValue); \
+			Serializer.Serialize(TEXTVIEW(JsonName), UnixTimestampValue); \
 			JsonDateTime = FDateTime::FromUnixTimestamp(UnixTimestampValue); \
 		} \
 		else \
 		{ \
 			int64 UnixTimestampValue = JsonDateTime.ToUnixTimestamp(); \
-			Serializer.Serialize(TEXT(JsonName), UnixTimestampValue); \
+			Serializer.Serialize(TEXTVIEW(JsonName), UnixTimestampValue); \
 		}
 
 #define JSON_SERIALIZE_DATETIME_UNIX_TIMESTAMP_MILLISECONDS(JsonName, JsonDateTime) \
 if (Serializer.IsLoading()) \
 { \
 	int64 UnixTimestampValueInMilliseconds; \
-	Serializer.Serialize(TEXT(JsonName), UnixTimestampValueInMilliseconds); \
+	Serializer.Serialize(TEXTVIEW(JsonName), UnixTimestampValueInMilliseconds); \
 	JsonDateTime = FDateTime::FromUnixTimestamp(UnixTimestampValueInMilliseconds / 1000); \
 } \
 else \
 { \
 	int64 UnixTimestampValueInMilliseconds = JsonDateTime.ToUnixTimestamp() * 1000; \
-	Serializer.Serialize(TEXT(JsonName), UnixTimestampValueInMilliseconds); \
+	Serializer.Serialize(TEXTVIEW(JsonName), UnixTimestampValueInMilliseconds); \
 }
 
 #define JSON_SERIALIZE_ENUM(JsonName, JsonEnum) \
 if (Serializer.IsLoading()) \
 { \
 	FString JsonTextValue; \
-	Serializer.Serialize(TEXT(JsonName), JsonTextValue); \
+	Serializer.Serialize(TEXTVIEW(JsonName), JsonTextValue); \
 	LexFromString(JsonEnum, *JsonTextValue); \
 } \
 else \
 { \
 	FString JsonTextValue = LexToString(JsonEnum); \
-	Serializer.Serialize(TEXT(JsonName), JsonTextValue); \
+	Serializer.Serialize(TEXTVIEW(JsonName), JsonTextValue); \
 }
 
 #define JSON_SERIALIZE_ENUM_ARRAY(JsonName, JsonArray, EnumType) \
 if (Serializer.IsLoading()) \
 { \
-	if (Serializer.GetObject()->HasTypedField<EJson::Array>(JsonName)) \
+	if (Serializer.GetObject()->HasTypedField<EJson::Array>(TEXTVIEW(JsonName))) \
 	{ \
 		EnumType EnumValue; \
-		for (auto It = Serializer.GetObject()->GetArrayField(JsonName).CreateConstIterator(); It; ++It) \
+		for (auto It = Serializer.GetObject()->GetArrayField(TEXTVIEW(JsonName)).CreateConstIterator(); It; ++It) \
 		{ \
 			LexFromString(EnumValue, *(*It)->AsString()); \
 			JsonArray.Add(EnumValue); \
@@ -337,11 +337,11 @@ if (Serializer.IsLoading()) \
 } \
 else \
 { \
-	Serializer.StartArray(JsonName); \
+	Serializer.StartArray(TEXTVIEW(JsonName)); \
 	for (EnumType& EnumValue : JsonArray) \
 	{ \
 		FString JsonTextValue = LexToString(EnumValue); \
-		Serializer.Serialize(TEXT(JsonName), JsonTextValue); \
+		Serializer.Serialize(TEXTVIEW(JsonName), JsonTextValue); \
 	} \
 	Serializer.EndArray(); \
 }
@@ -368,34 +368,51 @@ struct FJsonSerializerBase
 	virtual bool IsLoading() const = 0;
 	virtual bool IsSaving() const = 0;
 	virtual void StartObject() = 0;
-	virtual void StartObject(const FString& Name) = 0;
+	virtual void StartObject(FStringView Name) = 0;
 	virtual void EndObject() = 0;
 	virtual void StartArray() = 0;
-	virtual void StartArray(const FString& Name) = 0;
+	virtual void StartArray(FStringView Name) = 0;
 	virtual void EndArray() = 0;
-	virtual void Serialize(const TCHAR* Name, int32& Value) = 0;
-	virtual void Serialize(const TCHAR* Name, uint32& Value) = 0;
-	virtual void Serialize(const TCHAR* Name, int64& Value) = 0;
-	virtual void Serialize(const TCHAR* Name, bool& Value) = 0;
-	virtual void Serialize(const TCHAR* Name, FString& Value) = 0;
-	virtual void Serialize(const TCHAR* Name, FText& Value) = 0;
-	virtual void Serialize(const TCHAR* Name, float& Value) = 0;
-	virtual void Serialize(const TCHAR* Name, double& Value) = 0;
-	virtual void Serialize(const TCHAR* Name, FDateTime& Value) = 0;
+	virtual void Serialize(FStringView Name, int32& Value) = 0;
+	virtual void Serialize(FStringView Name, uint32& Value) = 0;
+	virtual void Serialize(FStringView Name, int64& Value) = 0;
+	virtual void Serialize(FStringView Name, bool& Value) = 0;
+	virtual void Serialize(FStringView Name, FString& Value) = 0;
+	virtual void Serialize(FStringView Name, FText& Value) = 0;
+	virtual void Serialize(FStringView Name, float& Value) = 0;
+	virtual void Serialize(FStringView Name, double& Value) = 0;
+	virtual void Serialize(FStringView Name, FDateTime& Value) = 0;
 	virtual void SerializeArray(FJsonSerializableArray& Array) = 0;
-	virtual void SerializeArray(const TCHAR* Name, FJsonSerializableArray& Value) = 0;
-	virtual void SerializeArray(const TCHAR* Name, FJsonSerializableArrayInt& Value) = 0;
-	virtual void SerializeArray(const TCHAR* Name, FJsonSerializableArrayFloat& Value) = 0;
-	virtual void SerializeMap(const TCHAR* Name, FJsonSerializableKeyValueMap& Map) = 0;
-	virtual void SerializeMap(const TCHAR* Name, FJsonSerializableKeyValueMapInt& Map) = 0;
-	virtual void SerializeMap(const TCHAR* Name, FJsonSerializableKeyValueMapArrayInt& Map) = 0;
-	virtual void SerializeMap(const TCHAR* Name, FJsonSerializableKeyValueMapInt64& Map) = 0;
-	virtual void SerializeMap(const TCHAR* Name, FJsonSerializableKeyValueMapFloat& Map) = 0;
+	virtual void SerializeArray(FStringView Name, FJsonSerializableArray& Value) = 0;
+	virtual void SerializeArray(FStringView Name, FJsonSerializableArrayInt& Value) = 0;
+	virtual void SerializeArray(FStringView Name, FJsonSerializableArrayFloat& Value) = 0;
+	virtual void SerializeMap(FStringView Name, FJsonSerializableKeyValueMap& Map) = 0;
+	virtual void SerializeMap(FStringView Name, FJsonSerializableKeyValueMapInt& Map) = 0;
+	virtual void SerializeMap(FStringView Name, FJsonSerializableKeyValueMapArrayInt& Map) = 0;
+	virtual void SerializeMap(FStringView Name, FJsonSerializableKeyValueMapInt64& Map) = 0;
+	virtual void SerializeMap(FStringView Name, FJsonSerializableKeyValueMapFloat& Map) = 0;
 	virtual void SerializeSimpleMap(FJsonSerializableKeyValueMap& Map) = 0;
-	virtual void SerializeMapSafe(const TCHAR* Name, FJsonSerializableKeyValueMap& Map) = 0;
+	virtual void SerializeMapSafe(FStringView Name, FJsonSerializableKeyValueMap& Map) = 0;
 	virtual TSharedPtr<FJsonObject> GetObject() = 0;
-	virtual void WriteIdentifierPrefix(const TCHAR* Name) = 0;
-	virtual void WriteRawJSONValue(const TCHAR* Value) = 0;
+	virtual void WriteIdentifierPrefix(FStringView Name) = 0;
+	virtual void WriteRawJSONValue(FStringView Value) = 0;
+
+#if !PLATFORM_TCHAR_IS_UTF8CHAR
+
+	UE_DEPRECATED(5.4, "Passing an ANSI string to StartObject has been deprecated outside of UTF-8 mode. Please use the overload that takes a TCHAR string.")
+	void StartObject(FAnsiStringView Name)
+	{
+		StartObject(StringCast<TCHAR>(Name.GetData(), Name.Len()));
+	}
+
+	UE_DEPRECATED(5.4, "Passing an ANSI string to StartArray has been deprecated outside of UTF-8 mode. Please use the overload that takes a TCHAR string.")
+	void StartArray(FAnsiStringView Name)
+	{
+		StartArray(StringCast<TCHAR>(Name.GetData(), Name.Len()));
+	}
+
+#endif // !PLATFORM_TCHAR_IS_UTF8CHAR
+	
 };
 
 /**
@@ -442,7 +459,7 @@ public:
 	/**
 	 * Starts a new object "{"
 	 */
-	virtual void StartObject(const FString& Name) override
+	virtual void StartObject(FStringView Name) override
 	{
 		JsonWriter->WriteObjectStart(Name);
 	}
@@ -459,7 +476,7 @@ public:
 		JsonWriter->WriteArrayStart();
 	}
 
-	virtual void StartArray(const FString& Name) override
+	virtual void StartArray(FStringView Name) override
 	{
 		JsonWriter->WriteArrayStart(Name);
 	}
@@ -474,7 +491,7 @@ public:
 	 * @param Name the field name to write out
 	 * @param Value the value to write out
 	 */
-	virtual void Serialize(const TCHAR* Name, int32& Value) override
+	virtual void Serialize(FStringView Name, int32& Value) override
 	{
 		JsonWriter->WriteValue(Name, Value);
 	}
@@ -484,7 +501,7 @@ public:
 	 * @param Name the field name to write out
 	 * @param Value the value to write out
 	 */
-	virtual void Serialize(const TCHAR* Name, uint32& Value) override
+	virtual void Serialize(FStringView Name, uint32& Value) override
 	{
 		JsonWriter->WriteValue(Name, static_cast<int64>(Value));
 	}
@@ -494,7 +511,7 @@ public:
 	 * @param Name the field name to write out
 	 * @param Value the value to write out
 	 */
-	virtual void Serialize(const TCHAR* Name, int64& Value) override
+	virtual void Serialize(FStringView Name, int64& Value) override
 	{
 		JsonWriter->WriteValue(Name, Value);
 	}
@@ -504,7 +521,7 @@ public:
 	 * @param Name the field name to write out
 	 * @param Value the value to write out
 	 */
-	virtual void Serialize(const TCHAR* Name, bool& Value) override
+	virtual void Serialize(FStringView Name, bool& Value) override
 	{
 		JsonWriter->WriteValue(Name, Value);
 	}
@@ -514,7 +531,7 @@ public:
 	 * @param Name the field name to write out
 	 * @param Value the value to write out
 	 */
-	virtual void Serialize(const TCHAR* Name, FString& Value) override
+	virtual void Serialize(FStringView Name, FString& Value) override
 	{
 		JsonWriter->WriteValue(Name, Value);
 	}
@@ -524,7 +541,7 @@ public:
 	 * @param Name the field name to write out
 	 * @param Value the value to write out
 	 */
-	virtual void Serialize(const TCHAR* Name, FText& Value) override
+	virtual void Serialize(FStringView Name, FText& Value) override
 	{
 		JsonWriter->WriteValue(Name, Value.ToString());
 	}
@@ -534,7 +551,7 @@ public:
 	 * @param Name the field name to write out
 	 * @param Value the value to write out
 	 */
-	virtual void Serialize(const TCHAR* Name, float& Value) override
+	virtual void Serialize(FStringView Name, float& Value) override
 	{
 		JsonWriter->WriteValue(Name, Value);
 	}
@@ -544,7 +561,7 @@ public:
 	* @param Name the field name to write out
 	* @param Value the value to write out
 	*/
-	virtual void Serialize(const TCHAR* Name, double& Value) override
+	virtual void Serialize(FStringView Name, double& Value) override
 	{
 		JsonWriter->WriteValue(Name, Value);
 	}
@@ -554,7 +571,7 @@ public:
 	* @param Name the field name to write out
 	* @param Value the value to write out
 	*/
-	virtual void Serialize(const TCHAR* Name, FDateTime& Value) override
+	virtual void Serialize(FStringView Name, FDateTime& Value) override
 	{
 		if (Value.GetTicks() > 0)
 		{
@@ -583,7 +600,7 @@ public:
 	 * @param Name the name of the property to serialize
 	 * @param Array the array to serialize
 	 */
-	virtual void SerializeArray(const TCHAR* Name, FJsonSerializableArray& Array) override
+	virtual void SerializeArray(FStringView Name, FJsonSerializableArray& Array) override
 	{
 		JsonWriter->WriteArrayStart(Name);
 		// Iterate all of values
@@ -599,7 +616,7 @@ public:
 	 * @param Name the name of the property to serialize
 	 * @param Array the array to serialize
 	 */
-	virtual void SerializeArray(const TCHAR* Name, FJsonSerializableArrayInt& Array) override
+	virtual void SerializeArray(FStringView Name, FJsonSerializableArrayInt& Array) override
 	{
 		JsonWriter->WriteArrayStart(Name);
 		// Iterate all of values
@@ -616,7 +633,7 @@ public:
 	 * @param Name the name of the property to serialize
 	 * @param Array the array to serialize
 	 */
-	virtual void SerializeArray(const TCHAR* Name, FJsonSerializableArrayFloat& Array) override
+	virtual void SerializeArray(FStringView Name, FJsonSerializableArrayFloat& Array) override
 	{
 		JsonWriter->WriteArrayStart(Name);
 		// Iterate all of values
@@ -633,7 +650,7 @@ public:
 	 * @param Name the name of the property to serialize
 	 * @param Map the map to serialize
 	 */
-	virtual void SerializeMap(const TCHAR* Name, FJsonSerializableKeyValueMap& Map) override
+	virtual void SerializeMap(FStringView Name, FJsonSerializableKeyValueMap& Map) override
 	{
 		JsonWriter->WriteObjectStart(Name);
 		// Iterate all of the keys and their values
@@ -650,7 +667,7 @@ public:
 	 * @param Name the name of the property to serialize
 	 * @param Map the map to serialize
 	 */
-	virtual void SerializeMap(const TCHAR* Name, FJsonSerializableKeyValueMapInt& Map) override
+	virtual void SerializeMap(FStringView Name, FJsonSerializableKeyValueMapInt& Map) override
 	{
 		JsonWriter->WriteObjectStart(Name);
 		// Iterate all of the keys and their values
@@ -666,7 +683,7 @@ public:
 	 * @param Name the name of the property to serialize
 	 * @param Map the map to serialize
 	 */
-	virtual void SerializeMap(const TCHAR* Name, FJsonSerializableKeyValueMapArrayInt& Map) override
+	virtual void SerializeMap(FStringView Name, FJsonSerializableKeyValueMapArrayInt& Map) override
 	{
 		JsonWriter->WriteObjectStart(Name);
 		// Iterate all of the keys and their values
@@ -683,7 +700,7 @@ public:
 	 * @param Name the name of the property to serialize
 	 * @param Map the map to serialize
 	 */
-	virtual void SerializeMap(const TCHAR* Name, FJsonSerializableKeyValueMapInt64& Map) override
+	virtual void SerializeMap(FStringView Name, FJsonSerializableKeyValueMapInt64& Map) override
 	{
 		JsonWriter->WriteObjectStart(Name);
 		// Iterate all of the keys and their values
@@ -700,7 +717,7 @@ public:
 	 * @param Name the name of the property to serialize
 	 * @param Map the map to serialize
 	 */
-	virtual void SerializeMap(const TCHAR* Name, FJsonSerializableKeyValueMapFloat& Map) override
+	virtual void SerializeMap(FStringView Name, FJsonSerializableKeyValueMapFloat& Map) override
 	{
 		JsonWriter->WriteObjectStart(Name);
 		// Iterate all of the keys and their values
@@ -723,17 +740,17 @@ public:
 	 * @param Name Name of property to serialize
 	 * @param Map The Map to copy String values from
 	 */
-	virtual void SerializeMapSafe(const TCHAR* Name, FJsonSerializableKeyValueMap& Map)
+	virtual void SerializeMapSafe(FStringView Name, FJsonSerializableKeyValueMap& Map)
 	{
 		SerializeMap(Name, Map);
 	}
 
-	virtual void WriteIdentifierPrefix(const TCHAR* Name)
+	virtual void WriteIdentifierPrefix(FStringView Name)
 	{
 		JsonWriter->WriteIdentifierPrefix(Name);
 	}
 
-	virtual void WriteRawJSONValue(const TCHAR* Value)
+	virtual void WriteRawJSONValue(FStringView Value)
 	{
 		JsonWriter->WriteRawJSONValue(Value);
 	}
@@ -776,7 +793,7 @@ public:
 		// Empty on purpose
 	}
 	/** Ignored */
-	virtual void StartObject(const FString& Name) override
+	virtual void StartObject(FStringView Name) override
 	{
 		// Empty on purpose
 	}
@@ -791,7 +808,7 @@ public:
 		// Empty on purpose
 	}
 	/** Ignored */
-	virtual void StartArray(const FString& Name) override
+	virtual void StartArray(FStringView Name) override
 	{
 		// Empty on purpose
 	}
@@ -806,7 +823,7 @@ public:
 	 * @param Name the name of the field to read
 	 * @param Value the out value to read the data into
 	 */
-	virtual void Serialize(const TCHAR* Name, int32& Value) override
+	virtual void Serialize(FStringView Name, int32& Value) override
 	{
 		if (JsonObject->HasTypedField<EJson::Number>(Name))
 		{
@@ -819,7 +836,7 @@ public:
 	 * @param Name the name of the field to read
 	 * @param Value the out value to read the data into
 	 */
-	virtual void Serialize(const TCHAR* Name, uint32& Value) override
+	virtual void Serialize(FStringView Name, uint32& Value) override
 	{
 		if (JsonObject->HasTypedField<EJson::Number>(Name))
 		{
@@ -832,7 +849,7 @@ public:
 	 * @param Name the name of the field to read
 	 * @param Value the out value to read the data into
 	 */
-	virtual void Serialize(const TCHAR* Name, int64& Value) override
+	virtual void Serialize(FStringView Name, int64& Value) override
 	{
 		if (JsonObject->HasTypedField<EJson::Number>(Name))
 		{
@@ -845,7 +862,7 @@ public:
 	 * @param Name the name of the field to read
 	 * @param Value the out value to read the data into
 	 */
-	virtual void Serialize(const TCHAR* Name, bool& Value) override
+	virtual void Serialize(FStringView Name, bool& Value) override
 	{
 		if (JsonObject->HasTypedField<EJson::Boolean>(Name))
 		{
@@ -858,7 +875,7 @@ public:
 	 * @param Name the name of the field to read
 	 * @param Value the out value to read the data into
 	 */
-	virtual void Serialize(const TCHAR* Name, FString& Value) override
+	virtual void Serialize(FStringView Name, FString& Value) override
 	{
 		if (JsonObject->HasTypedField<EJson::String>(Name))
 		{
@@ -871,7 +888,7 @@ public:
 	 * @param Name the name of the field to read
 	 * @param Value the out value to read the data into
 	 */
-	virtual void Serialize(const TCHAR* Name, FText& Value) override
+	virtual void Serialize(FStringView Name, FText& Value) override
 	{
 		if (JsonObject->HasTypedField<EJson::String>(Name))
 		{
@@ -884,7 +901,7 @@ public:
 	 * @param Name the name of the field to read
 	 * @param Value the out value to read the data into
 	 */
-	virtual void Serialize(const TCHAR* Name, float& Value) override
+	virtual void Serialize(FStringView Name, float& Value) override
 	{
 		if (JsonObject->HasTypedField<EJson::Number>(Name))
 		{
@@ -897,7 +914,7 @@ public:
 	* @param Name the name of the field to read
 	* @param Value the out value to read the data into
 	*/
-	virtual void Serialize(const TCHAR* Name, double& Value) override
+	virtual void Serialize(FStringView Name, double& Value) override
 	{
 		if (JsonObject->HasTypedField<EJson::Number>(Name))
 		{
@@ -910,7 +927,7 @@ public:
 	* @param Name the field name to write out
 	* @param Value the value to write out
 	*/
-	virtual void Serialize(const TCHAR* Name, FDateTime& Value) override
+	virtual void Serialize(FStringView Name, FDateTime& Value) override
 	{
 		if (JsonObject->HasTypedField<EJson::String>(Name))
 		{
@@ -934,7 +951,7 @@ public:
 	 * @param Name the name of the property to serialize
 	 * @param Array the array to serialize
 	 */
-	virtual void SerializeArray(const TCHAR* Name, FJsonSerializableArray& Array) override
+	virtual void SerializeArray(FStringView Name, FJsonSerializableArray& Array) override
 	{
 		if (JsonObject->HasTypedField<EJson::Array>(Name))
 		{
@@ -952,7 +969,7 @@ public:
 	 * @param Name the name of the property to serialize
 	 * @param Array the array to serialize
 	 */
-	virtual void SerializeArray(const TCHAR* Name, FJsonSerializableArrayInt& Array) override
+	virtual void SerializeArray(FStringView Name, FJsonSerializableArrayInt& Array) override
 	{
 		if (JsonObject->HasTypedField<EJson::Array>(Name))
 		{
@@ -970,7 +987,7 @@ public:
 	 * @param Name the name of the property to serialize
 	 * @param Array the array to serialize
 	 */
-	virtual void SerializeArray(const TCHAR* Name, FJsonSerializableArrayFloat& Array) override
+	virtual void SerializeArray(FStringView Name, FJsonSerializableArrayFloat& Array) override
 	{
 		if (JsonObject->HasTypedField<EJson::Array>(Name))
 		{
@@ -988,7 +1005,7 @@ public:
 	 * @param Name the name of the property to serialize
 	 * @param Map the map to serialize
 	 */
-	virtual void SerializeMap(const TCHAR* Name, FJsonSerializableKeyValueMap& Map) override
+	virtual void SerializeMap(FStringView Name, FJsonSerializableKeyValueMap& Map) override
 	{
 		if (JsonObject->HasTypedField<EJson::Object>(Name))
 		{
@@ -1007,7 +1024,7 @@ public:
 	 * @param Name the name of the property to serialize
 	 * @param Map the map to serialize
 	 */
-	virtual void SerializeMap(const TCHAR* Name, FJsonSerializableKeyValueMapInt& Map) override
+	virtual void SerializeMap(FStringView Name, FJsonSerializableKeyValueMapInt& Map) override
 	{
 		if (JsonObject->HasTypedField<EJson::Object>(Name))
 		{
@@ -1027,7 +1044,7 @@ public:
      * @param Name the name of the property to serialize
      * @param Map the map to serialize
      */
-    virtual void SerializeMap(const TCHAR* Name, FJsonSerializableKeyValueMapArrayInt& Map) override
+    virtual void SerializeMap(FStringView Name, FJsonSerializableKeyValueMapArrayInt& Map) override
     {
     	if (JsonObject->HasTypedField<EJson::Object>(Name))
     	{
@@ -1055,7 +1072,7 @@ public:
 	 * @param Name the name of the property to serialize
 	 * @param Map the map to serialize
 	 */
-	virtual void SerializeMap(const TCHAR* Name, FJsonSerializableKeyValueMapInt64& Map) override
+	virtual void SerializeMap(FStringView Name, FJsonSerializableKeyValueMapInt64& Map) override
 	{
 		if (JsonObject->HasTypedField<EJson::Object>(Name))
 		{
@@ -1075,7 +1092,7 @@ public:
 	 * @param Name the name of the property to serialize
 	 * @param Map the map to serialize
 	 */
-	virtual void SerializeMap(const TCHAR* Name, FJsonSerializableKeyValueMapFloat& Map) override
+	virtual void SerializeMap(FStringView Name, FJsonSerializableKeyValueMapFloat& Map) override
 	{
 		if (JsonObject->HasTypedField<EJson::Object>(Name))
 		{
@@ -1108,7 +1125,7 @@ public:
 	 * @param Name Name of property to deserialize
 	 * @param Map The Map to fill with String values found
 	 */
-	virtual void SerializeMapSafe(const TCHAR* Name, FJsonSerializableKeyValueMap& Map) override
+	virtual void SerializeMapSafe(FStringView Name, FJsonSerializableKeyValueMap& Map) override
 	{
 		if (JsonObject->HasTypedField<EJson::Object>(Name))
 		{
@@ -1125,13 +1142,13 @@ public:
 		}
 	}
 
-	virtual void WriteIdentifierPrefix(const TCHAR* Name)
+	virtual void WriteIdentifierPrefix(FStringView Name)
 	{
 		// Should never be called on a reader
 		check(false);
 	}
 
-	virtual void WriteRawJSONValue(const TCHAR* Value)
+	virtual void WriteRawJSONValue(FStringView Value)
 	{
 		// Should never be called on a reader
 		check(false);
diff --git a/Engine/Source/Runtime/Json/Public/Serialization/JsonWriter.h b/Engine/Source/Runtime/Json/Public/Serialization/JsonWriter.h
--- a/Engine/Source/Runtime/Json/Public/Serialization/JsonWriter.h
+++ b/Engine/Source/Runtime/Json/Public/Serialization/JsonWriter.h
@@ -282,7 +282,7 @@ public:
 
 	// WARNING: THIS IS DANGEROUS. Use this only if you know for a fact that the Value is valid JSON!
 	// Use this to insert the results of a different JSON Writer in.
-	void WriteRawJSONValue(const FString& Value)
+	void WriteRawJSONValue(FStringView Value)
 	{
 		check(CanWriteValueWithoutIdentifier());
 		WriteCommaIfNeeded();
