Refactor: move classes out of JsonSerializerMacros.h

diff --git a/Engine/Source/Runtime/Json/Private/Serialization/JsonDataBag.cpp b/Engine/Source/Runtime/Json/Private/Serialization/JsonDataBag.cpp
new file mode 100644
--- /dev/null
+++ b/Engine/Source/Runtime/Json/Private/Serialization/JsonDataBag.cpp
@@ -0,0 +1,116 @@
+// Copyright Epic Games, Inc. All Rights Reserved.
+
+#include "Serialization/JsonDataBag.h"
+#include "Serialization/JsonSerializer.h"
+#include "Serialization/JsonSerializerBase.h"
+
+void FJsonDataBag::Serialize(FJsonSerializerBase& Serializer, bool bFlatObject)
+{
+	if (Serializer.IsLoading())
+	{
+		// just grab a reference to the underlying JSON object
+		JsonObject = Serializer.GetObject();
+	}
+	else
+	{
+		if (!bFlatObject)
+		{
+			Serializer.StartObject();
+		}
+
+		if (JsonObject.IsValid())
+		{
+			for (const auto& It : JsonObject->Values)
+			{
+				TSharedPtr<FJsonValue> JsonValue = It.Value;
+				if (JsonValue.IsValid())
+				{
+					switch (JsonValue->Type)
+					{
+						case EJson::Boolean:
+						{
+							auto Value = JsonValue->AsBool();
+							Serializer.Serialize(*It.Key, Value);
+							break;
+						}
+						case EJson::Number:
+						{
+							auto Value = JsonValue->AsNumber();
+							Serializer.Serialize(*It.Key, Value);
+							break;
+						}
+						case EJson::String:
+						{
+							auto Value = JsonValue->AsString();
+							Serializer.Serialize(*It.Key, Value);
+							break;
+						}
+						case EJson::Array:
+						{
+							// if we have an array, serialize to string and write raw
+							FString JsonStr;
+							auto Writer = TJsonWriterFactory<TCHAR, TCondensedJsonPrintPolicy<TCHAR>>::Create(&JsonStr);
+							FJsonSerializer::Serialize(JsonValue->AsArray(), Writer);
+							Serializer.WriteIdentifierPrefix(*It.Key);
+							Serializer.WriteRawJSONValue(*JsonStr);
+							break;
+						}
+						case EJson::Object:
+						{
+							// if we have an object, serialize to string and write raw
+							FString JsonStr;
+							auto Writer = TJsonWriterFactory<TCHAR, TCondensedJsonPrintPolicy<TCHAR>>::Create(&JsonStr);
+							FJsonSerializer::Serialize(JsonValue->AsObject().ToSharedRef(), Writer);
+							// too bad there's no JsonObject serialization method on FJsonSerializerBase directly :-/
+							Serializer.WriteIdentifierPrefix(*It.Key);
+							Serializer.WriteRawJSONValue(*JsonStr);
+							break;
+						}
+					}
+				}
+			}
+		}
+
+		if (!bFlatObject)
+		{
+			Serializer.EndObject();
+		}
+	}
+}
+
+double FJsonDataBag::GetDouble(const FString& Key) const
+{
+	const auto Json = GetField(Key);
+	return Json.IsValid() ? Json->AsNumber() : 0.0;
+}
+
+FString FJsonDataBag::GetString(const FString& Key) const
+{
+	const auto Json = GetField(Key);
+	return Json.IsValid() ? Json->AsString() : FString();
+}
+
+bool FJsonDataBag::GetBool(const FString& Key) const
+{
+	const auto Json = GetField(Key);
+	return Json.IsValid() ? Json->AsBool() : false;
+}
+
+TSharedPtr<const FJsonValue> FJsonDataBag::GetField(const FString& Key) const
+{
+	if (JsonObject.IsValid())
+	{
+		return JsonObject->TryGetField(Key);
+	}
+	return TSharedPtr<const FJsonValue>();
+}
+
+void FJsonDataBag::SetFieldJson(const FString& Key, const TSharedPtr<FJsonValue>& Value)
+{
+	if (!JsonObject.IsValid())
+	{
+		JsonObject = MakeShared<FJsonObject>();
+	}
+	JsonObject->SetField(Key, Value);
+}
+
diff --git a/Engine/Source/Runtime/Json/Private/Serialization/JsonSerializable.cpp b/Engine/Source/Runtime/Json/Private/Serialization/JsonSerializable.cpp
new file mode 100644
--- /dev/null
+++ b/Engine/Source/Runtime/Json/Private/Serialization/JsonSerializable.cpp
@@ -0,0 +1,79 @@
+// Copyright Epic Games, Inc. All Rights Reserved.
+
+#include "Serialization/JsonSerializable.h"
+#include "Serialization/JsonSerializerReader.h"
+#include "Serialization/JsonSerializerWriter.h"
+
+FJsonSerializable::~FJsonSerializable() 
+{
+}
+
+const FString FJsonSerializable::ToJson(bool bPrettyPrint/*=true*/) const
+{
+	// Strip away const, because we use a single method that can read/write which requires non-const semantics
+	// Otherwise, we'd have to have 2 separate macros for declaring const to json and non-const from json
+	return ((FJsonSerializable*)this)->ToJson(bPrettyPrint);
+}
+
+const FString FJsonSerializable::ToJson(bool bPrettyPrint/*=true*/)
+{
+	FString JsonStr;
+	if (bPrettyPrint)
+	{
+		TSharedRef<TJsonWriter<> > JsonWriter = TJsonWriterFactory<>::Create(&JsonStr);
+		FJsonSerializerWriter<> Serializer(JsonWriter);
+		Serialize(Serializer, false);
+		JsonWriter->Close();
+	}
+	else
+	{
+		TSharedRef< TJsonWriter< TCHAR, TCondensedJsonPrintPolicy< TCHAR > > > JsonWriter = TJsonWriterFactory< TCHAR, TCondensedJsonPrintPolicy< TCHAR > >::Create( &JsonStr );
+		FJsonSerializerWriter<TCHAR, TCondensedJsonPrintPolicy< TCHAR >> Serializer(JsonWriter);
+		Serialize(Serializer, false);
+		JsonWriter->Close();
+	}
+	return JsonStr;
+}
+
+void FJsonSerializable::ToJson(TSharedRef<TJsonWriter<> >& JsonWriter, bool bFlatObject) const
+{
+	FJsonSerializerWriter<> Serializer(JsonWriter);
+	((FJsonSerializable*)this)->Serialize(Serializer, bFlatObject);
+}
+
+void FJsonSerializable::ToJson(TSharedRef< TJsonWriter< TCHAR, TCondensedJsonPrintPolicy< TCHAR > > >& JsonWriter, bool bFlatObject) const
+{
+	FJsonSerializerWriter<TCHAR, TCondensedJsonPrintPolicy< TCHAR >> Serializer(JsonWriter);
+	((FJsonSerializable*)this)->Serialize(Serializer, bFlatObject);
+}
+
+bool FJsonSerializable::FromJson(const FString& Json)
+{
+	return FromJson(CopyTemp(Json));
+}
+
+bool FJsonSerializable::FromJson(FString&& Json)
+{
+	TSharedPtr<FJsonObject> JsonObject;
+	TSharedRef<TJsonReader<> > JsonReader = TJsonReaderFactory<>::Create(MoveTemp(Json));
+	if (FJsonSerializer::Deserialize(JsonReader,JsonObject) &&
+		JsonObject.IsValid())
+	{
+		FJsonSerializerReader Serializer(JsonObject);
+		Serialize(Serializer, false);
+		return true;
+	}
+	UE_LOG(LogJson, Warning, TEXT("Failed to parse Json from a string: %s"), *JsonReader->GetErrorMessage());
+	return false;
+}
+
+bool FJsonSerializable::FromJson(TSharedPtr<FJsonObject> JsonObject)
+{
+	if (JsonObject.IsValid())
+	{
+		FJsonSerializerReader Serializer(JsonObject);
+		Serialize(Serializer, false);
+		return true;
+	}
+	return false;
+}
diff --git a/Engine/Source/Runtime/Json/Private/Serialization/JsonSerializerReader.cpp b/Engine/Source/Runtime/Json/Private/Serialization/JsonSerializerReader.cpp
new file mode 100644
--- /dev/null
+++ b/Engine/Source/Runtime/Json/Private/Serialization/JsonSerializerReader.cpp
@@ -0,0 +1,293 @@
+// Copyright Epic Games, Inc. All Rights Reserved.
+
+#include "Serialization/JsonSerializerReader.h"
+
+FJsonSerializerReader::FJsonSerializerReader(TSharedPtr<FJsonObject> InJsonObject)
+	: JsonObject(InJsonObject)
+{
+}
+
+FJsonSerializerReader::~FJsonSerializerReader()
+{
+}
+
+bool FJsonSerializerReader::IsLoading() const 
+{ 
+	return true; 
+}
+
+bool FJsonSerializerReader::IsSaving() const 
+{ 
+	return false; 
+}
+
+TSharedPtr<FJsonObject> FJsonSerializerReader::GetObject() 
+{ 
+	return JsonObject; 
+}
+
+void FJsonSerializerReader::StartObject() 
+{
+	// Empty on purpose
+}
+
+void FJsonSerializerReader::StartObject(FStringView Name) 
+{
+	// Empty on purpose
+}
+
+void FJsonSerializerReader::EndObject() 
+{
+	// Empty on purpose
+}
+
+void FJsonSerializerReader::StartArray() 
+{
+	// Empty on purpose
+}
+
+void FJsonSerializerReader::StartArray(FStringView Name) 
+{
+	// Empty on purpose
+}
+
+void FJsonSerializerReader::EndArray() 
+{
+	// Empty on purpose
+}
+
+void FJsonSerializerReader::Serialize(FStringView Name, int32& Value) 
+{
+	if (JsonObject->HasTypedField<EJson::Number>(Name))
+	{
+		JsonObject->TryGetNumberField(Name, Value);
+	}
+}
+
+void FJsonSerializerReader::Serialize(FStringView Name, uint32& Value) 
+{
+	if (JsonObject->HasTypedField<EJson::Number>(Name))
+	{
+		JsonObject->TryGetNumberField(Name, Value);
+	}
+}
+
+void FJsonSerializerReader::Serialize(FStringView Name, int64& Value) 
+{
+	if (JsonObject->HasTypedField<EJson::Number>(Name))
+	{
+		JsonObject->TryGetNumberField(Name, Value);
+	}
+}
+
+void FJsonSerializerReader::Serialize(FStringView Name, bool& Value) 
+{
+	if (JsonObject->HasTypedField<EJson::Boolean>(Name))
+	{
+		Value = JsonObject->GetBoolField(Name);
+	}
+}
+
+void FJsonSerializerReader::Serialize(FStringView Name, FString& Value) 
+{
+	if (JsonObject->HasTypedField<EJson::String>(Name))
+	{
+		Value = JsonObject->GetStringField(Name);
+	}
+}
+
+void FJsonSerializerReader::Serialize(FStringView Name, FText& Value) 
+{
+	if (JsonObject->HasTypedField<EJson::String>(Name))
+	{
+		Value = FText::FromString(JsonObject->GetStringField(Name));
+	}
+}
+
+void FJsonSerializerReader::Serialize(FStringView Name, float& Value)
+{
+	if (JsonObject->HasTypedField<EJson::Number>(Name))
+	{
+		Value = (float)JsonObject->GetNumberField(Name);
+	}
+}
+
+void FJsonSerializerReader::Serialize(FStringView Name, double& Value) 
+{
+	if (JsonObject->HasTypedField<EJson::Number>(Name))
+	{
+		Value = JsonObject->GetNumberField(Name);
+	}
+}
+
+void FJsonSerializerReader::Serialize(FStringView Name, FDateTime& Value) 
+{
+	if (JsonObject->HasTypedField<EJson::String>(Name))
+	{
+		FDateTime::ParseIso8601(*JsonObject->GetStringField(Name), Value);
+	}
+}
+
+void FJsonSerializerReader::SerializeArray(FJsonSerializableArray& Array) 
+{
+	// @todo - higher level serialization is expecting a Json Object
+	check(0 && TEXT("Not implemented"));
+}
+
+void FJsonSerializerReader::SerializeArray(FStringView Name, FJsonSerializableArray& Array) 
+{
+	if (JsonObject->HasTypedField<EJson::Array>(Name))
+	{
+		TArray< TSharedPtr<FJsonValue> > JsonArray = JsonObject->GetArrayField(Name);
+		// Iterate all of the keys and their values
+		for (TSharedPtr<FJsonValue>& Value : JsonArray)
+		{
+			Array.Add(Value->AsString());
+		}
+	}
+}
+
+void FJsonSerializerReader::SerializeArray(FStringView Name, FJsonSerializableArrayInt& Array) 
+{
+	if (JsonObject->HasTypedField<EJson::Array>(Name))
+	{
+		TArray< TSharedPtr<FJsonValue> > JsonArray = JsonObject->GetArrayField(Name);
+		// Iterate all of the keys and their values
+		for (TSharedPtr<FJsonValue>& Value : JsonArray)
+		{
+			Array.Add((int32)Value->AsNumber());
+		}
+	}
+}
+
+void FJsonSerializerReader::SerializeArray(FStringView Name, FJsonSerializableArrayFloat& Array) 
+{
+	if (JsonObject->HasTypedField<EJson::Array>(Name))
+	{
+		TArray< TSharedPtr<FJsonValue> > JsonArray = JsonObject->GetArrayField(Name);
+		// Iterate all of the keys and their values
+		for (TSharedPtr<FJsonValue>& Value : JsonArray)
+		{
+			Array.Add((float)Value->AsNumber());
+		}
+	}
+}
+
+void FJsonSerializerReader::SerializeMap(FStringView Name, FJsonSerializableKeyValueMap& Map) 
+{
+	if (JsonObject->HasTypedField<EJson::Object>(Name))
+	{
+		TSharedPtr<FJsonObject> JsonMap = JsonObject->GetObjectField(Name);
+		// Iterate all of the keys and their values
+		for (const TPair<FString, TSharedPtr<FJsonValue>>& Pair : JsonMap->Values)
+		{
+			Map.Add(Pair.Key, Pair.Value->AsString());
+		}
+	}
+}
+
+void FJsonSerializerReader::SerializeMap(FStringView Name, FJsonSerializableKeyValueMapInt& Map) 
+{
+	if (JsonObject->HasTypedField<EJson::Object>(Name))
+	{
+		TSharedPtr<FJsonObject> JsonMap = JsonObject->GetObjectField(Name);
+		// Iterate all of the keys and their values
+		for (const TPair<FString, TSharedPtr<FJsonValue>>& Pair : JsonMap->Values)
+		{
+			const int32 Value = (int32)Pair.Value->AsNumber();
+			Map.Add(Pair.Key, Value);
+		}
+	}
+}
+
+void FJsonSerializerReader::SerializeMap(FStringView Name, FJsonSerializableKeyValueMapArrayInt& Map) 
+{
+	if (JsonObject->HasTypedField<EJson::Object>(Name))
+	{
+		TSharedPtr<FJsonObject> JsonMap = JsonObject->GetObjectField(Name);
+		// Iterate all of the keys and their values
+		for (const TPair<FString, TSharedPtr<FJsonValue>>& Pair : JsonMap->Values)
+		{
+			FJsonSerializableArrayInt TempArray;
+			for (const TSharedPtr<FJsonValue>& ArrayValue : Pair.Value->AsArray())
+			{
+				int32 IntValue;
+				if (ArrayValue.IsValid() && ArrayValue->TryGetNumber(IntValue))
+				{
+					TempArray.Add(IntValue);
+				}
+			}
+			Map.Add(Pair.Key, TempArray);
+		}
+	}
+}
+
+void FJsonSerializerReader::SerializeMap(FStringView Name, FJsonSerializableKeyValueMapInt64& Map) 
+{
+	if (JsonObject->HasTypedField<EJson::Object>(Name))
+	{
+		TSharedPtr<FJsonObject> JsonMap = JsonObject->GetObjectField(Name);
+		// Iterate all of the keys and their values
+		for (const TPair<FString, TSharedPtr<FJsonValue>>& Pair : JsonMap->Values)
+		{
+			const int64 Value = (int64)Pair.Value->AsNumber();
+			Map.Add(Pair.Key, Value);
+		}
+	}
+}
+
+void FJsonSerializerReader::SerializeMap(FStringView Name, FJsonSerializableKeyValueMapFloat& Map) 
+{
+	if (JsonObject->HasTypedField<EJson::Object>(Name))
+	{
+		TSharedPtr<FJsonObject> JsonMap = JsonObject->GetObjectField(Name);
+		// Iterate all of the keys and their values
+		for (const TPair<FString, TSharedPtr<FJsonValue>>& Pair : JsonMap->Values)
+		{
+			const float Value = (float)Pair.Value->AsNumber();
+			Map.Add(Pair.Key, Value);
+		}
+	}
+}
+
+void FJsonSerializerReader::SerializeSimpleMap(FJsonSerializableKeyValueMap& Map) 
+{
+	// Iterate all of the keys and their values, only taking simple types (not array/object), all in string form
+	for (auto KeyValueIt = JsonObject->Values.CreateConstIterator(); KeyValueIt; ++KeyValueIt)
+	{
+		FString Value;
+		if (KeyValueIt.Value()->TryGetString(Value))
+		{
+			Map.Add(KeyValueIt.Key(), MoveTemp(Value));
+		}
+	}
+}
+
+void FJsonSerializerReader::SerializeMapSafe(FStringView Name, FJsonSerializableKeyValueMap& Map) 
+{
+	if (JsonObject->HasTypedField<EJson::Object>(Name))
+	{
+		// Iterate all of the keys and their values, only taking simple types (not array/object), all in string form
+		TSharedPtr<FJsonObject> JsonMap = JsonObject->GetObjectField(Name);
+		for (const TPair<FString, TSharedPtr<FJsonValue>>& Pair : JsonMap->Values)
+		{
+			FString Value;
+			if (Pair.Value->TryGetString(Value))
+			{
+				Map.Add(Pair.Key, MoveTemp(Value));
+			}
+		}
+	}
+}
+
+void FJsonSerializerReader::WriteIdentifierPrefix(FStringView Name)
+{
+	// Should never be called on a reader
+	check(false);
+}
+
+void FJsonSerializerReader::WriteRawJSONValue(FStringView Value)
+{
+	// Should never be called on a reader
+	check(false);
+}
diff --git a/Engine/Source/Runtime/Json/Public/Serialization/JsonDataBag.h b/Engine/Source/Runtime/Json/Public/Serialization/JsonDataBag.h
new file mode 100644
--- /dev/null
+++ b/Engine/Source/Runtime/Json/Public/Serialization/JsonDataBag.h
@@ -0,0 +1,34 @@
+// Copyright Epic Games, Inc. All Rights Reserved.
+
+#pragma once
+
+#include "Serialization/JsonSerializable.h"
+
+class FJsonSerializerReader;
+
+/**
+ * Useful if you just want access to the underlying FJsonObject (for cases where the schema is loose or an outer system will do further de/serialization)
+ */
+struct FJsonDataBag : public FJsonSerializable
+{
+	JSON_API virtual void Serialize(FJsonSerializerBase& Serializer, bool bFlatObject) override;
+
+	JSON_API double GetDouble(const FString& Key) const;
+
+	JSON_API FString GetString(const FString& Key) const;
+
+	JSON_API bool GetBool(const FString& Key) const;
+
+	JSON_API TSharedPtr<const FJsonValue> GetField(const FString& Key) const;
+
+	template<typename JSON_TYPE, typename Arg>
+	void SetField(const FString& Key, Arg&& Value)
+	{
+		SetFieldJson(Key, MakeShared<JSON_TYPE>(Forward<Arg>(Value)));
+	}
+
+	JSON_API void SetFieldJson(const FString& Key, const TSharedPtr<FJsonValue>& Value);
+
+public:
+	TSharedPtr<FJsonObject> JsonObject;
+};
diff --git a/Engine/Source/Runtime/Json/Public/Serialization/JsonSerializable.h b/Engine/Source/Runtime/Json/Public/Serialization/JsonSerializable.h
new file mode 100644
--- /dev/null
+++ b/Engine/Source/Runtime/Json/Public/Serialization/JsonSerializable.h
@@ -0,0 +1,76 @@
+// Copyright Epic Games, Inc. All Rights Reserved.
+
+#pragma once
+
+#include "Dom/JsonObject.h"
+#include "Policies/CondensedJsonPrintPolicy.h"
+#include "Serialization/JsonWriter.h"
+#include "Serialization/JsonSerializer.h"
+#include "Serialization/JsonSerializerReader.h"
+
+/**
+ * Base class for a JSON serializable object
+ */
+struct FJsonSerializable
+{
+	/**
+	 *	Virtualize destructor as we provide overridable functions
+	 */
+	JSON_API virtual ~FJsonSerializable();
+
+	/**
+	 * Used to allow serialization of a const ref
+	 *
+	 * @return the corresponding json string
+	 */
+	JSON_API const FString ToJson(bool bPrettyPrint = true) const;
+	/**
+	 * Serializes this object to its JSON string form
+	 *
+	 * @param bPrettyPrint - If true, will use the pretty json formatter
+	 * @return the corresponding json string
+	 */
+	JSON_API virtual const FString ToJson(bool bPrettyPrint=true);
+	JSON_API virtual void ToJson(TSharedRef<TJsonWriter<> >& JsonWriter, bool bFlatObject) const;
+	JSON_API virtual void ToJson(TSharedRef< TJsonWriter< TCHAR, TCondensedJsonPrintPolicy< TCHAR > > >& JsonWriter, bool bFlatObject) const;
+
+	/**
+	 * Serializes the contents of a JSON string into this object
+	 *
+	 * @param Json the JSON data to serialize from
+	 */
+	JSON_API virtual bool FromJson(const FString& Json);
+
+	/**
+	 * Serializes the contents of a JSON string into this object
+	 *
+	 * @param Json the JSON data to serialize from
+	 */
+	JSON_API virtual bool FromJson(FString&& Json);
+
+	template <class CharType = TCHAR>
+	bool FromJson(TStringView<CharType> Json)
+	{
+		TSharedPtr<FJsonObject> JsonObject;
+		TSharedRef<TJsonReader<CharType> > JsonReader = TJsonReaderFactory<CharType>::CreateFromView(Json);
+		if (FJsonSerializer::Deserialize(JsonReader,JsonObject) &&
+			JsonObject.IsValid())
+		{
+			FJsonSerializerReader Serializer(JsonObject);
+			Serialize(Serializer, false);
+			return true;
+		}
+		UE_LOG(LogJson, Warning, TEXT("Failed to parse Json from a string: %s"), *JsonReader->GetErrorMessage());
+		return false;
+	}
+
+	JSON_API virtual bool FromJson(TSharedPtr<FJsonObject> JsonObject);
+
+	/**
+	 * Abstract method that needs to be supplied using the macros
+	 *
+	 * @param Serializer the object that will perform serialization in/out of JSON
+	 * @param bFlatObject if true then no object wrapper is used
+	 */
+	JSON_API virtual void Serialize(FJsonSerializerBase& Serializer, bool bFlatObject) = 0;
+};
diff --git a/Engine/Source/Runtime/Json/Public/Serialization/JsonSerializerBase.h b/Engine/Source/Runtime/Json/Public/Serialization/JsonSerializerBase.h
new file mode 100644
--- /dev/null
+++ b/Engine/Source/Runtime/Json/Public/Serialization/JsonSerializerBase.h
@@ -0,0 +1,60 @@
+// Copyright Epic Games, Inc. All Rights Reserved.
+
+#pragma once
+
+#include "Serialization/JsonTypes.h"
+#include "Dom/JsonObject.h"
+
+/**
+ * Base interface used to serialize to/from JSON. Hides the fact there are separate read/write classes
+ */
+struct FJsonSerializerBase
+{
+	virtual bool IsLoading() const = 0;
+	virtual bool IsSaving() const = 0;
+	virtual void StartObject() = 0;
+	virtual void StartObject(FStringView Name) = 0;
+	virtual void EndObject() = 0;
+	virtual void StartArray() = 0;
+	virtual void StartArray(FStringView Name) = 0;
+	virtual void EndArray() = 0;
+	virtual void Serialize(FStringView Name, int32& Value) = 0;
+	virtual void Serialize(FStringView Name, uint32& Value) = 0;
+	virtual void Serialize(FStringView Name, int64& Value) = 0;
+	virtual void Serialize(FStringView Name, bool& Value) = 0;
+	virtual void Serialize(FStringView Name, FString& Value) = 0;
+	virtual void Serialize(FStringView Name, FText& Value) = 0;
+	virtual void Serialize(FStringView Name, float& Value) = 0;
+	virtual void Serialize(FStringView Name, double& Value) = 0;
+	virtual void Serialize(FStringView Name, FDateTime& Value) = 0;
+	virtual void SerializeArray(FJsonSerializableArray& Array) = 0;
+	virtual void SerializeArray(FStringView Name, FJsonSerializableArray& Value) = 0;
+	virtual void SerializeArray(FStringView Name, FJsonSerializableArrayInt& Value) = 0;
+	virtual void SerializeArray(FStringView Name, FJsonSerializableArrayFloat& Value) = 0;
+	virtual void SerializeMap(FStringView Name, FJsonSerializableKeyValueMap& Map) = 0;
+	virtual void SerializeMap(FStringView Name, FJsonSerializableKeyValueMapInt& Map) = 0;
+	virtual void SerializeMap(FStringView Name, FJsonSerializableKeyValueMapArrayInt& Map) = 0;
+	virtual void SerializeMap(FStringView Name, FJsonSerializableKeyValueMapInt64& Map) = 0;
+	virtual void SerializeMap(FStringView Name, FJsonSerializableKeyValueMapFloat& Map) = 0;
+	virtual void SerializeSimpleMap(FJsonSerializableKeyValueMap& Map) = 0;
+	virtual void SerializeMapSafe(FStringView Name, FJsonSerializableKeyValueMap& Map) = 0;
+	virtual TSharedPtr<FJsonObject> GetObject() = 0;
+	virtual void WriteIdentifierPrefix(FStringView Name) = 0;
+	virtual void WriteRawJSONValue(FStringView Value) = 0;
+
+#if !PLATFORM_TCHAR_IS_UTF8CHAR
+
+	UE_DEPRECATED(5.4, "Passing an ANSI string to StartObject has been deprecated outside of UTF-8 mode. Please use the overload that takes a TCHAR string.")
+	void StartObject(FAnsiStringView Name)
+	{
+		StartObject(StringCast<TCHAR>(Name.GetData(), Name.Len()));
+	}
+
+	UE_DEPRECATED(5.4, "Passing an ANSI string to StartArray has been deprecated outside of UTF-8 mode. Please use the overload that takes a TCHAR string.")
+	void StartArray(FAnsiStringView Name)
+	{
+		StartArray(StringCast<TCHAR>(Name.GetData(), Name.Len()));
+	}
+
+#endif // !PLATFORM_TCHAR_IS_UTF8CHAR
+};
diff --git a/Engine/Source/Runtime/Json/Public/Serialization/JsonSerializerMacros.h b/Engine/Source/Runtime/Json/Public/Serialization/JsonSerializerMacros.h
--- a/Engine/Source/Runtime/Json/Public/Serialization/JsonSerializerMacros.h
+++ b/Engine/Source/Runtime/Json/Public/Serialization/JsonSerializerMacros.h
@@ -5,9 +5,15 @@
 #include "CoreMinimal.h"
 #include "Policies/PrettyJsonPrintPolicy.h"
 #include "Policies/CondensedJsonPrintPolicy.h"
+#include "Serialization/JsonDataBag.h"
 #include "Serialization/JsonTypes.h"
 #include "Serialization/JsonReader.h"
 #include "Serialization/JsonSerializer.h"
+#include "Serialization/JsonSerializable.h"
+#include "Serialization/JsonSerializerBase.h"
+#include "Serialization/JsonSerializerReader.h"
+#include "Serialization/JsonSerializerWriter.h"
+#include "JsonGlobals.h"
 
 /**
  * Macros used to generate a serialization function for a class derived from FJsonSerializable
@@ -345,1061 +351,3 @@ else \
 	} \
 	Serializer.EndArray(); \
 }
-
-struct FJsonSerializerBase;
-
-/** Array of data */
-typedef TArray<FString> FJsonSerializableArray;
-typedef TArray<int32> FJsonSerializableArrayInt;
-typedef TArray<float> FJsonSerializableArrayFloat;
-
-/** Maps a key to a value */
-typedef TMap<FString, FString> FJsonSerializableKeyValueMap;
-typedef TMap<FString, int32> FJsonSerializableKeyValueMapInt;
-typedef TMap<FString, int64> FJsonSerializableKeyValueMapInt64;
-typedef TMap<FString, float> FJsonSerializableKeyValueMapFloat;
-typedef TMap<FString, FJsonSerializableArrayInt> FJsonSerializableKeyValueMapArrayInt;
-
-/**
- * Base interface used to serialize to/from JSON. Hides the fact there are separate read/write classes
- */
-struct FJsonSerializerBase
-{
-	virtual bool IsLoading() const = 0;
-	virtual bool IsSaving() const = 0;
-	virtual void StartObject() = 0;
-	virtual void StartObject(FStringView Name) = 0;
-	virtual void EndObject() = 0;
-	virtual void StartArray() = 0;
-	virtual void StartArray(FStringView Name) = 0;
-	virtual void EndArray() = 0;
-	virtual void Serialize(FStringView Name, int32& Value) = 0;
-	virtual void Serialize(FStringView Name, uint32& Value) = 0;
-	virtual void Serialize(FStringView Name, int64& Value) = 0;
-	virtual void Serialize(FStringView Name, bool& Value) = 0;
-	virtual void Serialize(FStringView Name, FString& Value) = 0;
-	virtual void Serialize(FStringView Name, FText& Value) = 0;
-	virtual void Serialize(FStringView Name, float& Value) = 0;
-	virtual void Serialize(FStringView Name, double& Value) = 0;
-	virtual void Serialize(FStringView Name, FDateTime& Value) = 0;
-	virtual void SerializeArray(FJsonSerializableArray& Array) = 0;
-	virtual void SerializeArray(FStringView Name, FJsonSerializableArray& Value) = 0;
-	virtual void SerializeArray(FStringView Name, FJsonSerializableArrayInt& Value) = 0;
-	virtual void SerializeArray(FStringView Name, FJsonSerializableArrayFloat& Value) = 0;
-	virtual void SerializeMap(FStringView Name, FJsonSerializableKeyValueMap& Map) = 0;
-	virtual void SerializeMap(FStringView Name, FJsonSerializableKeyValueMapInt& Map) = 0;
-	virtual void SerializeMap(FStringView Name, FJsonSerializableKeyValueMapArrayInt& Map) = 0;
-	virtual void SerializeMap(FStringView Name, FJsonSerializableKeyValueMapInt64& Map) = 0;
-	virtual void SerializeMap(FStringView Name, FJsonSerializableKeyValueMapFloat& Map) = 0;
-	virtual void SerializeSimpleMap(FJsonSerializableKeyValueMap& Map) = 0;
-	virtual void SerializeMapSafe(FStringView Name, FJsonSerializableKeyValueMap& Map) = 0;
-	virtual TSharedPtr<FJsonObject> GetObject() = 0;
-	virtual void WriteIdentifierPrefix(FStringView Name) = 0;
-	virtual void WriteRawJSONValue(FStringView Value) = 0;
-
-#if !PLATFORM_TCHAR_IS_UTF8CHAR
-
-	UE_DEPRECATED(5.4, "Passing an ANSI string to StartObject has been deprecated outside of UTF-8 mode. Please use the overload that takes a TCHAR string.")
-	void StartObject(FAnsiStringView Name)
-	{
-		StartObject(StringCast<TCHAR>(Name.GetData(), Name.Len()));
-	}
-
-	UE_DEPRECATED(5.4, "Passing an ANSI string to StartArray has been deprecated outside of UTF-8 mode. Please use the overload that takes a TCHAR string.")
-	void StartArray(FAnsiStringView Name)
-	{
-		StartArray(StringCast<TCHAR>(Name.GetData(), Name.Len()));
-	}
-
-#endif // !PLATFORM_TCHAR_IS_UTF8CHAR
-	
-};
-
-/**
- * Implements the abstract serializer interface hiding the underlying writer object
- */
-template <class CharType = TCHAR, class PrintPolicy = TPrettyJsonPrintPolicy<CharType> >
-class FJsonSerializerWriter :
-	public FJsonSerializerBase
-{
-	/** The object to write the JSON output to */
-	TSharedRef<TJsonWriter<CharType, PrintPolicy> > JsonWriter;
-
-public:
-
-	/**
-	 * Initializes the writer object
-	 *
-	 * @param InJsonWriter the object to write the JSON output to
-	 */
-	FJsonSerializerWriter(TSharedRef<TJsonWriter<CharType, PrintPolicy> > InJsonWriter) :
-		JsonWriter(InJsonWriter)
-	{
-	}
-
-	virtual ~FJsonSerializerWriter()
-	{
-	}
-
-	/** Is the JSON being read from */
-	virtual bool IsLoading() const override { return false; }
-	/** Is the JSON being written to */
-	virtual bool IsSaving() const override { return true; }
-	/** Access to the root object */
-	virtual TSharedPtr<FJsonObject> GetObject() override { return TSharedPtr<FJsonObject>(); }
-
-	/**
-	 * Starts a new object "{"
-	 */
-	virtual void StartObject() override
-	{
-		JsonWriter->WriteObjectStart();
-	}
-
-	/**
-	 * Starts a new object "{"
-	 */
-	virtual void StartObject(FStringView Name) override
-	{
-		JsonWriter->WriteObjectStart(Name);
-	}
-	/**
-	 * Completes the definition of an object "}"
-	 */
-	virtual void EndObject() override
-	{
-		JsonWriter->WriteObjectEnd();
-	}
-
-	virtual void StartArray() override
-	{
-		JsonWriter->WriteArrayStart();
-	}
-
-	virtual void StartArray(FStringView Name) override
-	{
-		JsonWriter->WriteArrayStart(Name);
-	}
-
-	virtual void EndArray() override
-	{
-		JsonWriter->WriteArrayEnd();
-	}
-	/**
-	 * Writes the field name and the corresponding value to the JSON data
-	 *
-	 * @param Name the field name to write out
-	 * @param Value the value to write out
-	 */
-	virtual void Serialize(FStringView Name, int32& Value) override
-	{
-		JsonWriter->WriteValue(Name, Value);
-	}
-	/**
-	 * Writes the field name and the corresponding value to the JSON data
-	 *
-	 * @param Name the field name to write out
-	 * @param Value the value to write out
-	 */
-	virtual void Serialize(FStringView Name, uint32& Value) override
-	{
-		JsonWriter->WriteValue(Name, static_cast<int64>(Value));
-	}
-	/**
-	 * Writes the field name and the corresponding value to the JSON data
-	 *
-	 * @param Name the field name to write out
-	 * @param Value the value to write out
-	 */
-	virtual void Serialize(FStringView Name, int64& Value) override
-	{
-		JsonWriter->WriteValue(Name, Value);
-	}
-	/**
-	 * Writes the field name and the corresponding value to the JSON data
-	 *
-	 * @param Name the field name to write out
-	 * @param Value the value to write out
-	 */
-	virtual void Serialize(FStringView Name, bool& Value) override
-	{
-		JsonWriter->WriteValue(Name, Value);
-	}
-	/**
-	 * Writes the field name and the corresponding value to the JSON data
-	 *
-	 * @param Name the field name to write out
-	 * @param Value the value to write out
-	 */
-	virtual void Serialize(FStringView Name, FString& Value) override
-	{
-		JsonWriter->WriteValue(Name, Value);
-	}
-	/**
-	 * Writes the field name and the corresponding value to the JSON data
-	 *
-	 * @param Name the field name to write out
-	 * @param Value the value to write out
-	 */
-	virtual void Serialize(FStringView Name, FText& Value) override
-	{
-		JsonWriter->WriteValue(Name, Value.ToString());
-	}
-	/**
-	 * Writes the field name and the corresponding value to the JSON data
-	 *
-	 * @param Name the field name to write out
-	 * @param Value the value to write out
-	 */
-	virtual void Serialize(FStringView Name, float& Value) override
-	{
-		JsonWriter->WriteValue(Name, Value);
-	}
-	/**
-	* Writes the field name and the corresponding value to the JSON data
-	*
-	* @param Name the field name to write out
-	* @param Value the value to write out
-	*/
-	virtual void Serialize(FStringView Name, double& Value) override
-	{
-		JsonWriter->WriteValue(Name, Value);
-	}
-	/**
-	* Writes the field name and the corresponding value to the JSON data
-	*
-	* @param Name the field name to write out
-	* @param Value the value to write out
-	*/
-	virtual void Serialize(FStringView Name, FDateTime& Value) override
-	{
-		if (Value.GetTicks() > 0)
-		{
-			JsonWriter->WriteValue(Name, Value.ToIso8601());
-		}
-	}
-	/**
-	 * Serializes an array of values
-	 *
-	 * @param Name the name of the property to serialize
-	 * @param Array the array to serialize
-	 */
-	virtual void SerializeArray(FJsonSerializableArray& Array) override
-	{
-		JsonWriter->WriteArrayStart();
-		// Iterate all of values
-		for (FJsonSerializableArray::TIterator ArrayIt(Array); ArrayIt; ++ArrayIt)
-		{
-			JsonWriter->WriteValue(*ArrayIt);
-		}
-		JsonWriter->WriteArrayEnd();
-	}
-	/**
-	 * Serializes an array of values with an identifier
-	 *
-	 * @param Name the name of the property to serialize
-	 * @param Array the array to serialize
-	 */
-	virtual void SerializeArray(FStringView Name, FJsonSerializableArray& Array) override
-	{
-		JsonWriter->WriteArrayStart(Name);
-		// Iterate all of values
-		for (FJsonSerializableArray::ElementType& Item :  Array)
-		{
-			JsonWriter->WriteValue(Item);
-		}
-		JsonWriter->WriteArrayEnd();
-	}
-	/**
-	 * Serializes an array of values with an identifier
-	 *
-	 * @param Name the name of the property to serialize
-	 * @param Array the array to serialize
-	 */
-	virtual void SerializeArray(FStringView Name, FJsonSerializableArrayInt& Array) override
-	{
-		JsonWriter->WriteArrayStart(Name);
-		// Iterate all of values
-		for (FJsonSerializableArrayInt::ElementType& Item : Array)
-		{
-			JsonWriter->WriteValue(Item);
-		}
-		JsonWriter->WriteArrayEnd();
-	}
-
-	/**
-	 * Serializes an array of values with an identifier
-	 *
-	 * @param Name the name of the property to serialize
-	 * @param Array the array to serialize
-	 */
-	virtual void SerializeArray(FStringView Name, FJsonSerializableArrayFloat& Array) override
-	{
-		JsonWriter->WriteArrayStart(Name);
-		// Iterate all of values
-		for (FJsonSerializableArrayFloat::ElementType& Item : Array)
-		{
-			JsonWriter->WriteValue(Item);
-		}
-		JsonWriter->WriteArrayEnd();
-	}
-
-	/**
-	 * Serializes the keys & values for map
-	 *
-	 * @param Name the name of the property to serialize
-	 * @param Map the map to serialize
-	 */
-	virtual void SerializeMap(FStringView Name, FJsonSerializableKeyValueMap& Map) override
-	{
-		JsonWriter->WriteObjectStart(Name);
-		// Iterate all of the keys and their values
-		for (FJsonSerializableKeyValueMap::ElementType& Pair : Map)
-		{
-			Serialize(*Pair.Key, Pair.Value);
-		}
-		JsonWriter->WriteObjectEnd();
-	}
-
-	/**
-	 * Serializes the keys & values for map
-	 *
-	 * @param Name the name of the property to serialize
-	 * @param Map the map to serialize
-	 */
-	virtual void SerializeMap(FStringView Name, FJsonSerializableKeyValueMapInt& Map) override
-	{
-		JsonWriter->WriteObjectStart(Name);
-		// Iterate all of the keys and their values
-		for (FJsonSerializableKeyValueMapInt::ElementType& Pair : Map)
-		{
-			Serialize(*Pair.Key, Pair.Value);
-		}
-		JsonWriter->WriteObjectEnd();
-	}
-	/**
-	 * Serializes the keys & values for map
-	 *
-	 * @param Name the name of the property to serialize
-	 * @param Map the map to serialize
-	 */
-	virtual void SerializeMap(FStringView Name, FJsonSerializableKeyValueMapArrayInt& Map) override
-	{
-		JsonWriter->WriteObjectStart(Name);
-		// Iterate all of the keys and their values
-		for (FJsonSerializableKeyValueMapArrayInt::ElementType& Pair : Map)
-		{
-			SerializeArray(*Pair.Key, Pair.Value);
-		}
-		JsonWriter->WriteObjectEnd();
-	}
-
-	/**
-	 * Serializes the keys & values for map
-	 *
-	 * @param Name the name of the property to serialize
-	 * @param Map the map to serialize
-	 */
-	virtual void SerializeMap(FStringView Name, FJsonSerializableKeyValueMapInt64& Map) override
-	{
-		JsonWriter->WriteObjectStart(Name);
-		// Iterate all of the keys and their values
-		for (FJsonSerializableKeyValueMapInt64::ElementType& Pair : Map)
-		{
-			Serialize(*Pair.Key, Pair.Value);
-		}
-		JsonWriter->WriteObjectEnd();
-	}
-
-	/**
-	 * Serializes the keys & values for map
-	 *
-	 * @param Name the name of the property to serialize
-	 * @param Map the map to serialize
-	 */
-	virtual void SerializeMap(FStringView Name, FJsonSerializableKeyValueMapFloat& Map) override
-	{
-		JsonWriter->WriteObjectStart(Name);
-		// Iterate all of the keys and their values
-		for (FJsonSerializableKeyValueMapFloat::ElementType& Pair : Map)
-		{
-			Serialize(*Pair.Key, Pair.Value);
-		}
-		JsonWriter->WriteObjectEnd();
-	}
-
-	virtual void SerializeSimpleMap(FJsonSerializableKeyValueMap& Map) override
-	{
-		// writing does nothing here, this is meant to read in all data from a json object 
-		// writing is explicitly handled per key/type
-	}
-
-	/**
-	 * Serializes keys and values from an object into a map.
-	 *
-	 * @param Name Name of property to serialize
-	 * @param Map The Map to copy String values from
-	 */
-	virtual void SerializeMapSafe(FStringView Name, FJsonSerializableKeyValueMap& Map)
-	{
-		SerializeMap(Name, Map);
-	}
-
-	virtual void WriteIdentifierPrefix(FStringView Name)
-	{
-		JsonWriter->WriteIdentifierPrefix(Name);
-	}
-
-	virtual void WriteRawJSONValue(FStringView Value)
-	{
-		JsonWriter->WriteRawJSONValue(Value);
-	}
-};
-
-/**
- * Implements the abstract serializer interface hiding the underlying reader object
- */
-class FJsonSerializerReader :
-	public FJsonSerializerBase
-{
-	/** The object that holds the parsed JSON data */
-	TSharedPtr<FJsonObject> JsonObject;
-
-public:
-	/**
-	 * Inits the base JSON object that is being read from
-	 *
-	 * @param InJsonObject the JSON object to serialize from
-	 */
-	FJsonSerializerReader(TSharedPtr<FJsonObject> InJsonObject) :
-		JsonObject(InJsonObject)
-	{
-	}
-
-	virtual ~FJsonSerializerReader()
-	{
-	}
-
-	/** Is the JSON being read from */
-	virtual bool IsLoading() const override { return true; }
-	/** Is the JSON being written to */
-	virtual bool IsSaving() const override { return false; }
-	/** Access to the root Json object being read */
-	virtual TSharedPtr<FJsonObject> GetObject() override { return JsonObject; }
-
-	/** Ignored */
-	virtual void StartObject() override
-	{
-		// Empty on purpose
-	}
-	/** Ignored */
-	virtual void StartObject(FStringView Name) override
-	{
-		// Empty on purpose
-	}
-	/** Ignored */
-	virtual void EndObject() override
-	{
-		// Empty on purpose
-	}
-	/** Ignored */
-	virtual void StartArray() override
-	{
-		// Empty on purpose
-	}
-	/** Ignored */
-	virtual void StartArray(FStringView Name) override
-	{
-		// Empty on purpose
-	}
-	/** Ignored */
-	virtual void EndArray() override
-	{
-		// Empty on purpose
-	}
-	/**
-	 * If the underlying json object has the field, it is read into the value
-	 *
-	 * @param Name the name of the field to read
-	 * @param Value the out value to read the data into
-	 */
-	virtual void Serialize(FStringView Name, int32& Value) override
-	{
-		if (JsonObject->HasTypedField<EJson::Number>(Name))
-		{
-			JsonObject->TryGetNumberField(Name, Value);
-		}
-	}
-	/**
-	 * If the underlying json object has the field, it is read into the value
-	 *
-	 * @param Name the name of the field to read
-	 * @param Value the out value to read the data into
-	 */
-	virtual void Serialize(FStringView Name, uint32& Value) override
-	{
-		if (JsonObject->HasTypedField<EJson::Number>(Name))
-		{
-			JsonObject->TryGetNumberField(Name, Value);
-		}
-	}
-	/**
-	 * If the underlying json object has the field, it is read into the value
-	 *
-	 * @param Name the name of the field to read
-	 * @param Value the out value to read the data into
-	 */
-	virtual void Serialize(FStringView Name, int64& Value) override
-	{
-		if (JsonObject->HasTypedField<EJson::Number>(Name))
-		{
-			JsonObject->TryGetNumberField(Name, Value);
-		}
-	}
-	/**
-	 * If the underlying json object has the field, it is read into the value
-	 *
-	 * @param Name the name of the field to read
-	 * @param Value the out value to read the data into
-	 */
-	virtual void Serialize(FStringView Name, bool& Value) override
-	{
-		if (JsonObject->HasTypedField<EJson::Boolean>(Name))
-		{
-			Value = JsonObject->GetBoolField(Name);
-		}
-	}
-	/**
-	 * If the underlying json object has the field, it is read into the value
-	 *
-	 * @param Name the name of the field to read
-	 * @param Value the out value to read the data into
-	 */
-	virtual void Serialize(FStringView Name, FString& Value) override
-	{
-		if (JsonObject->HasTypedField<EJson::String>(Name))
-		{
-			Value = JsonObject->GetStringField(Name);
-		}
-	}
-	/**
-	 * If the underlying json object has the field, it is read into the value
-	 *
-	 * @param Name the name of the field to read
-	 * @param Value the out value to read the data into
-	 */
-	virtual void Serialize(FStringView Name, FText& Value) override
-	{
-		if (JsonObject->HasTypedField<EJson::String>(Name))
-		{
-			Value = FText::FromString(JsonObject->GetStringField(Name));
-		}
-	}
-	/**
-	 * If the underlying json object has the field, it is read into the value
-	 *
-	 * @param Name the name of the field to read
-	 * @param Value the out value to read the data into
-	 */
-	virtual void Serialize(FStringView Name, float& Value) override
-	{
-		if (JsonObject->HasTypedField<EJson::Number>(Name))
-		{
-			Value = (float)JsonObject->GetNumberField(Name);
-		}
-	}
-	/**
-	* If the underlying json object has the field, it is read into the value
-	*
-	* @param Name the name of the field to read
-	* @param Value the out value to read the data into
-	*/
-	virtual void Serialize(FStringView Name, double& Value) override
-	{
-		if (JsonObject->HasTypedField<EJson::Number>(Name))
-		{
-			Value = JsonObject->GetNumberField(Name);
-		}
-	}
-	/**
-	* Writes the field name and the corresponding value to the JSON data
-	*
-	* @param Name the field name to write out
-	* @param Value the value to write out
-	*/
-	virtual void Serialize(FStringView Name, FDateTime& Value) override
-	{
-		if (JsonObject->HasTypedField<EJson::String>(Name))
-		{
-			FDateTime::ParseIso8601(*JsonObject->GetStringField(Name), Value);
-		}
-	}
-	/**
-	 * Serializes an array of values
-	 *
-	 * @param Name the name of the property to serialize
-	 * @param Array the array to serialize
-	 */
-	virtual void SerializeArray(FJsonSerializableArray& Array) override
-	{
-		// @todo - higher level serialization is expecting a Json Object
-		check(0 && TEXT("Not implemented"));
-	}
-	/**
-	 * Serializes an array of values with an identifier
-	 *
-	 * @param Name the name of the property to serialize
-	 * @param Array the array to serialize
-	 */
-	virtual void SerializeArray(FStringView Name, FJsonSerializableArray& Array) override
-	{
-		if (JsonObject->HasTypedField<EJson::Array>(Name))
-		{
-			TArray< TSharedPtr<FJsonValue> > JsonArray = JsonObject->GetArrayField(Name);
-			// Iterate all of the keys and their values
-			for (TSharedPtr<FJsonValue>& Value : JsonArray)
-			{
-				Array.Add(Value->AsString());
-			}
-		}
-	}
-	/**
-	 * Serializes an array of values with an identifier
-	 *
-	 * @param Name the name of the property to serialize
-	 * @param Array the array to serialize
-	 */
-	virtual void SerializeArray(FStringView Name, FJsonSerializableArrayInt& Array) override
-	{
-		if (JsonObject->HasTypedField<EJson::Array>(Name))
-		{
-			TArray< TSharedPtr<FJsonValue> > JsonArray = JsonObject->GetArrayField(Name);
-			// Iterate all of the keys and their values
-			for (TSharedPtr<FJsonValue>& Value : JsonArray)
-			{
-				Array.Add((int32)Value->AsNumber());
-			}
-		}
-	}
-	/**
-	 * Serializes an array of values with an identifier
-	 *
-	 * @param Name the name of the property to serialize
-	 * @param Array the array to serialize
-	 */
-	virtual void SerializeArray(FStringView Name, FJsonSerializableArrayFloat& Array) override
-	{
-		if (JsonObject->HasTypedField<EJson::Array>(Name))
-		{
-			TArray< TSharedPtr<FJsonValue> > JsonArray = JsonObject->GetArrayField(Name);
-			// Iterate all of the keys and their values
-			for (TSharedPtr<FJsonValue>& Value : JsonArray)
-			{
-				Array.Add((float)Value->AsNumber());
-			}
-		}
-	}
-	/**
-	 * Serializes the keys & values for map
-	 *
-	 * @param Name the name of the property to serialize
-	 * @param Map the map to serialize
-	 */
-	virtual void SerializeMap(FStringView Name, FJsonSerializableKeyValueMap& Map) override
-	{
-		if (JsonObject->HasTypedField<EJson::Object>(Name))
-		{
-			TSharedPtr<FJsonObject> JsonMap = JsonObject->GetObjectField(Name);
-			// Iterate all of the keys and their values
-			for (const TPair<FString, TSharedPtr<FJsonValue>>& Pair : JsonMap->Values)
-			{
-				Map.Add(Pair.Key, Pair.Value->AsString());
-			}
-		}
-	}
-
-	/**
-	 * Serializes the keys & values for map
-	 *
-	 * @param Name the name of the property to serialize
-	 * @param Map the map to serialize
-	 */
-	virtual void SerializeMap(FStringView Name, FJsonSerializableKeyValueMapInt& Map) override
-	{
-		if (JsonObject->HasTypedField<EJson::Object>(Name))
-		{
-			TSharedPtr<FJsonObject> JsonMap = JsonObject->GetObjectField(Name);
-			// Iterate all of the keys and their values
-			for (const TPair<FString, TSharedPtr<FJsonValue>>& Pair : JsonMap->Values)
-			{
-				const int32 Value = (int32)Pair.Value->AsNumber();
-				Map.Add(Pair.Key, Value);
-			}
-		}
-	}
-	
-	/**
-     * Serializes the keys & values for map
-     *
-     * @param Name the name of the property to serialize
-     * @param Map the map to serialize
-     */
-    virtual void SerializeMap(FStringView Name, FJsonSerializableKeyValueMapArrayInt& Map) override
-    {
-    	if (JsonObject->HasTypedField<EJson::Object>(Name))
-    	{
-    		TSharedPtr<FJsonObject> JsonMap = JsonObject->GetObjectField(Name);
-    		// Iterate all of the keys and their values
-    		for (const TPair<FString, TSharedPtr<FJsonValue>>& Pair : JsonMap->Values)
-    		{
-    			FJsonSerializableArrayInt TempArray;
-    			for(const TSharedPtr<FJsonValue>& ArrayValue : Pair.Value->AsArray())
-    			{
-    				int32 IntValue;
-    				if(ArrayValue.IsValid() && ArrayValue->TryGetNumber(IntValue))
-    				{
-    					TempArray.Add(IntValue);
-    				}
-    			}
-    			Map.Add(Pair.Key, TempArray);
-    		}
-    	}
-    }
-
-	/**
-	 * Serializes the keys & values for map
-	 *
-	 * @param Name the name of the property to serialize
-	 * @param Map the map to serialize
-	 */
-	virtual void SerializeMap(FStringView Name, FJsonSerializableKeyValueMapInt64& Map) override
-	{
-		if (JsonObject->HasTypedField<EJson::Object>(Name))
-		{
-			TSharedPtr<FJsonObject> JsonMap = JsonObject->GetObjectField(Name);
-			// Iterate all of the keys and their values
-			for (const TPair<FString, TSharedPtr<FJsonValue>>& Pair : JsonMap->Values)
-			{
-				const int64 Value = (int64)Pair.Value->AsNumber();
-				Map.Add(Pair.Key, Value);
-			}
-		}
-	}
-
-	/**
-	 * Serializes the keys & values for map
-	 *
-	 * @param Name the name of the property to serialize
-	 * @param Map the map to serialize
-	 */
-	virtual void SerializeMap(FStringView Name, FJsonSerializableKeyValueMapFloat& Map) override
-	{
-		if (JsonObject->HasTypedField<EJson::Object>(Name))
-		{
-			TSharedPtr<FJsonObject> JsonMap = JsonObject->GetObjectField(Name);
-			// Iterate all of the keys and their values
-			for (const TPair<FString, TSharedPtr<FJsonValue>>& Pair : JsonMap->Values)
-			{
-				const float Value = (float)Pair.Value->AsNumber();
-				Map.Add(Pair.Key, Value);
-			}
-		}
-	}
-
-	virtual void SerializeSimpleMap(FJsonSerializableKeyValueMap& Map) override
-	{
-		// Iterate all of the keys and their values, only taking simple types (not array/object), all in string form
-		for (auto KeyValueIt = JsonObject->Values.CreateConstIterator(); KeyValueIt; ++KeyValueIt)
-		{
-			FString Value;
-			if (KeyValueIt.Value()->TryGetString(Value))
-			{
-				Map.Add(KeyValueIt.Key(), MoveTemp(Value));
-			}
-		}
-	}
-
-	/**
-	 * Deserializes keys and values from an object into a map, but only if the value is trivially convertable to string.
-	 *
-	 * @param Name Name of property to deserialize
-	 * @param Map The Map to fill with String values found
-	 */
-	virtual void SerializeMapSafe(FStringView Name, FJsonSerializableKeyValueMap& Map) override
-	{
-		if (JsonObject->HasTypedField<EJson::Object>(Name))
-		{
-			// Iterate all of the keys and their values, only taking simple types (not array/object), all in string form
-			TSharedPtr<FJsonObject> JsonMap = JsonObject->GetObjectField(Name);
-			for (const TPair<FString, TSharedPtr<FJsonValue>>& Pair : JsonMap->Values)
-			{
-				FString Value;
-				if (Pair.Value->TryGetString(Value))
-				{
-					Map.Add(Pair.Key, MoveTemp(Value));
-				}
-			}
-		}
-	}
-
-	virtual void WriteIdentifierPrefix(FStringView Name)
-	{
-		// Should never be called on a reader
-		check(false);
-	}
-
-	virtual void WriteRawJSONValue(FStringView Value)
-	{
-		// Should never be called on a reader
-		check(false);
-	}
-};
-
-/**
- * Base class for a JSON serializable object
- */
-struct FJsonSerializable
-{
-	/**
-	 *	Virtualize destructor as we provide overridable functions
-	 */
-	virtual ~FJsonSerializable() {}
-
-	/**
-	 * Used to allow serialization of a const ref
-	 *
-	 * @return the corresponding json string
-	 */
-	inline const FString ToJson(bool bPrettyPrint = true) const
-	{
-		// Strip away const, because we use a single method that can read/write which requires non-const semantics
-		// Otherwise, we'd have to have 2 separate macros for declaring const to json and non-const from json
-		return ((FJsonSerializable*)this)->ToJson(bPrettyPrint);
-	}
-	/**
-	 * Serializes this object to its JSON string form
-	 *
-	 * @param bPrettyPrint - If true, will use the pretty json formatter
-	 * @return the corresponding json string
-	 */
-	virtual const FString ToJson(bool bPrettyPrint=true)
-	{
-		FString JsonStr;
-		if (bPrettyPrint)
-		{
-			TSharedRef<TJsonWriter<> > JsonWriter = TJsonWriterFactory<>::Create(&JsonStr);
-			FJsonSerializerWriter<> Serializer(JsonWriter);
-			Serialize(Serializer, false);
-			JsonWriter->Close();
-		}
-		else
-		{
-			TSharedRef< TJsonWriter< TCHAR, TCondensedJsonPrintPolicy< TCHAR > > > JsonWriter = TJsonWriterFactory< TCHAR, TCondensedJsonPrintPolicy< TCHAR > >::Create( &JsonStr );
-			FJsonSerializerWriter<TCHAR, TCondensedJsonPrintPolicy< TCHAR >> Serializer(JsonWriter);
-			Serialize(Serializer, false);
-			JsonWriter->Close();
-		}
-		return JsonStr;
-	}
-	virtual void ToJson(TSharedRef<TJsonWriter<> >& JsonWriter, bool bFlatObject) const
-	{
-		FJsonSerializerWriter<> Serializer(JsonWriter);
-		((FJsonSerializable*)this)->Serialize(Serializer, bFlatObject);
-	}
-	virtual void ToJson(TSharedRef< TJsonWriter< TCHAR, TCondensedJsonPrintPolicy< TCHAR > > >& JsonWriter, bool bFlatObject) const
-	{
-		FJsonSerializerWriter<TCHAR, TCondensedJsonPrintPolicy< TCHAR >> Serializer(JsonWriter);
-		((FJsonSerializable*)this)->Serialize(Serializer, bFlatObject);
-	}
-
-	/**
-	 * Serializes the contents of a JSON string into this object
-	 *
-	 * @param Json the JSON data to serialize from
-	 */
-	virtual bool FromJson(const FString& Json)
-	{
-		return FromJson(CopyTemp(Json));
-	}
-
-	/**
-	 * Serializes the contents of a JSON string into this object
-	 *
-	 * @param Json the JSON data to serialize from
-	 */
-	virtual bool FromJson(FString&& Json)
-	{
-		TSharedPtr<FJsonObject> JsonObject;
-		TSharedRef<TJsonReader<> > JsonReader = TJsonReaderFactory<>::Create(MoveTemp(Json));
-		if (FJsonSerializer::Deserialize(JsonReader,JsonObject) &&
-			JsonObject.IsValid())
-		{
-			FJsonSerializerReader Serializer(JsonObject);
-			Serialize(Serializer, false);
-			return true;
-		}
-		UE_LOG(LogJson, Warning, TEXT("Failed to parse Json from a string: %s"), *JsonReader->GetErrorMessage());
-		return false;
-	}
-
-	template <class CharType = TCHAR>
-	bool FromJson(TStringView<CharType> Json)
-	{
-		TSharedPtr<FJsonObject> JsonObject;
-		TSharedRef<TJsonReader<CharType> > JsonReader = TJsonReaderFactory<CharType>::CreateFromView(Json);
-		if (FJsonSerializer::Deserialize(JsonReader,JsonObject) &&
-			JsonObject.IsValid())
-		{
-			FJsonSerializerReader Serializer(JsonObject);
-			Serialize(Serializer, false);
-			return true;
-		}
-		UE_LOG(LogJson, Warning, TEXT("Failed to parse Json from a string: %s"), *JsonReader->GetErrorMessage());
-		return false;
-	}
-
-	virtual bool FromJson(TSharedPtr<FJsonObject> JsonObject)
-	{
-		if (JsonObject.IsValid())
-		{
-			FJsonSerializerReader Serializer(JsonObject);
-			Serialize(Serializer, false);
-			return true;
-		}
-		return false;
-	}
-
-	/**
-	 * Abstract method that needs to be supplied using the macros
-	 *
-	 * @param Serializer the object that will perform serialization in/out of JSON
-	 * @param bFlatObject if true then no object wrapper is used
-	 */
-	virtual void Serialize(FJsonSerializerBase& Serializer, bool bFlatObject) = 0;
-};
-
-/**
- * Useful if you just want access to the underlying FJsonObject (for cases where the schema is loose or an outer system will do further de/serialization)
- */
-struct FJsonDataBag
-	: public FJsonSerializable
-{
-	virtual void Serialize(FJsonSerializerBase& Serializer, bool bFlatObject) override
-	{
-		if (Serializer.IsLoading())
-		{
-			// just grab a reference to the underlying JSON object
-			JsonObject = Serializer.GetObject();
-		}
-		else
-		{
-			if (!bFlatObject)
-			{
-				Serializer.StartObject();
-			}
-
-			if (JsonObject.IsValid())
-			{
-				for (const auto& It : JsonObject->Values)
-				{
-					TSharedPtr<FJsonValue> JsonValue = It.Value;
-					if (JsonValue.IsValid())
-					{
-						switch (JsonValue->Type)
-						{
-							case EJson::Boolean:
-							{
-								auto Value = JsonValue->AsBool();
-								Serializer.Serialize(*It.Key, Value);
-								break;
-							}
-							case EJson::Number:
-							{
-								auto Value = JsonValue->AsNumber();
-								Serializer.Serialize(*It.Key, Value);
-								break;
-							}
-							case EJson::String:
-							{
-								auto Value = JsonValue->AsString();
-								Serializer.Serialize(*It.Key, Value);
-								break;
-							}
-							case EJson::Array:
-							{
-								// if we have an array, serialize to string and write raw
-								FString JsonStr;
-								auto Writer = TJsonWriterFactory<TCHAR, TCondensedJsonPrintPolicy<TCHAR>>::Create(&JsonStr);
-								FJsonSerializer::Serialize(JsonValue->AsArray(), Writer);
-								Serializer.WriteIdentifierPrefix(*It.Key);
-								Serializer.WriteRawJSONValue(*JsonStr);
-								break;
-							}
-							case EJson::Object:
-							{
-								// if we have an object, serialize to string and write raw
-								FString JsonStr;
-								auto Writer = TJsonWriterFactory<TCHAR, TCondensedJsonPrintPolicy<TCHAR>>::Create(&JsonStr);
-								FJsonSerializer::Serialize(JsonValue->AsObject().ToSharedRef(), Writer);
-								// too bad there's no JsonObject serialization method on FJsonSerializerBase directly :-/
-								Serializer.WriteIdentifierPrefix(*It.Key);
-								Serializer.WriteRawJSONValue(*JsonStr);
-								break;
-							}
-						}
-					}
-				}
-			}
-
-			if (!bFlatObject)
-			{
-				Serializer.EndObject();
-			}
-		}
-	}
-
-	double GetDouble(const FString& Key) const
-	{
-		const auto Json = GetField(Key);
-		return Json.IsValid() ? Json->AsNumber() : 0.0;
-	}
-
-	FString GetString(const FString& Key) const
-	{
-		const auto Json = GetField(Key);
-		return Json.IsValid() ? Json->AsString() : FString();
-	}
-
-	bool GetBool(const FString& Key) const
-	{
-		const auto Json = GetField(Key);
-		return Json.IsValid() ? Json->AsBool() : false;
-	}
-
-	TSharedPtr<const FJsonValue> GetField(const FString& Key) const
-	{
-		if (JsonObject.IsValid())
-		{
-			return JsonObject->TryGetField(Key);
-		}
-		return TSharedPtr<const FJsonValue>();
-	}
-
-	template<typename JSON_TYPE, typename Arg>
-	void SetField(const FString& Key, Arg&& Value)
-	{
-		SetFieldJson(Key, MakeShared<JSON_TYPE>(Forward<Arg>(Value)));
-	}
-
-	void SetFieldJson(const FString& Key, const TSharedPtr<FJsonValue>& Value)
-	{
-		if (!JsonObject.IsValid())
-		{
-			JsonObject = MakeShared<FJsonObject>();
-		}
-		JsonObject->SetField(Key, Value);
-	}
-
-public:
-	TSharedPtr<FJsonObject> JsonObject;
-};
diff --git a/Engine/Source/Runtime/Json/Public/Serialization/JsonSerializerReader.h b/Engine/Source/Runtime/Json/Public/Serialization/JsonSerializerReader.h
new file mode 100644
--- /dev/null
+++ b/Engine/Source/Runtime/Json/Public/Serialization/JsonSerializerReader.h
@@ -0,0 +1,191 @@
+// Copyright Epic Games, Inc. All Rights Reserved.
+
+#pragma once
+
+#include "Serialization/JsonSerializerBase.h"
+#include "Dom/JsonObject.h"
+
+/**
+ * Implements the abstract serializer interface hiding the underlying reader object
+ */
+class FJsonSerializerReader : public FJsonSerializerBase
+{
+public:
+	/**
+	 * Inits the base JSON object that is being read from
+	 *
+	 * @param InJsonObject the JSON object to serialize from
+	 */
+	JSON_API FJsonSerializerReader(TSharedPtr<FJsonObject> InJsonObject);
+
+	JSON_API virtual ~FJsonSerializerReader();
+
+	/** Is the JSON being read from */
+	JSON_API virtual bool IsLoading() const override;
+	/** Is the JSON being written to */
+	JSON_API virtual bool IsSaving() const override;
+	/** Access to the root Json object being read */
+	JSON_API virtual TSharedPtr<FJsonObject> GetObject() override;
+
+	/** Ignored */
+	JSON_API virtual void StartObject() override;
+	/** Ignored */
+	JSON_API virtual void StartObject(FStringView Name) override;
+	/** Ignored */
+	JSON_API virtual void EndObject() override;
+	/** Ignored */
+	JSON_API virtual void StartArray() override;
+	/** Ignored */
+	JSON_API virtual void StartArray(FStringView Name) override;
+	/** Ignored */
+	JSON_API virtual void EndArray() override;
+	/**
+	 * If the underlying json object has the field, it is read into the value
+	 *
+	 * @param Name the name of the field to read
+	 * @param Value the out value to read the data into
+	 */
+	JSON_API virtual void Serialize(FStringView Name, int32& Value) override;
+	/**
+	 * If the underlying json object has the field, it is read into the value
+	 *
+	 * @param Name the name of the field to read
+	 * @param Value the out value to read the data into
+	 */
+	JSON_API virtual void Serialize(FStringView Name, uint32& Value) override;
+	/**
+	 * If the underlying json object has the field, it is read into the value
+	 *
+	 * @param Name the name of the field to read
+	 * @param Value the out value to read the data into
+	 */
+	JSON_API virtual void Serialize(FStringView Name, int64& Value) override;
+	/**
+	 * If the underlying json object has the field, it is read into the value
+	 *
+	 * @param Name the name of the field to read
+	 * @param Value the out value to read the data into
+	 */
+	JSON_API virtual void Serialize(FStringView Name, bool& Value) override;
+	/**
+	 * If the underlying json object has the field, it is read into the value
+	 *
+	 * @param Name the name of the field to read
+	 * @param Value the out value to read the data into
+	 */
+	JSON_API virtual void Serialize(FStringView Name, FString& Value) override;
+	/**
+	 * If the underlying json object has the field, it is read into the value
+	 *
+	 * @param Name the name of the field to read
+	 * @param Value the out value to read the data into
+	 */
+	JSON_API virtual void Serialize(FStringView Name, FText& Value) override;
+	/**
+	 * If the underlying json object has the field, it is read into the value
+	 *
+	 * @param Name the name of the field to read
+	 * @param Value the out value to read the data into
+	 */
+	JSON_API virtual void Serialize(FStringView Name, float& Value) override;
+	/**
+	* If the underlying json object has the field, it is read into the value
+	*
+	* @param Name the name of the field to read
+	* @param Value the out value to read the data into
+	*/
+	JSON_API virtual void Serialize(FStringView Name, double& Value) override;
+	/**
+	* Writes the field name and the corresponding value to the JSON data
+	*
+	* @param Name the field name to write out
+	* @param Value the value to write out
+	*/
+	JSON_API virtual void Serialize(FStringView Name, FDateTime& Value) override;
+	/**
+	 * Serializes an array of values
+	 *
+	 * @param Name the name of the property to serialize
+	 * @param Array the array to serialize
+	 */
+	JSON_API virtual void SerializeArray(FJsonSerializableArray& Array) override;
+	/**
+	 * Serializes an array of values with an identifier
+	 *
+	 * @param Name the name of the property to serialize
+	 * @param Array the array to serialize
+	 */
+	JSON_API virtual void SerializeArray(FStringView Name, FJsonSerializableArray& Array) override;
+	/**
+	 * Serializes an array of values with an identifier
+	 *
+	 * @param Name the name of the property to serialize
+	 * @param Array the array to serialize
+	 */
+	JSON_API virtual void SerializeArray(FStringView Name, FJsonSerializableArrayInt& Array) override;
+	/**
+	 * Serializes an array of values with an identifier
+	 *
+	 * @param Name the name of the property to serialize
+	 * @param Array the array to serialize
+	 */
+	JSON_API virtual void SerializeArray(FStringView Name, FJsonSerializableArrayFloat& Array) override;
+	/**
+	 * Serializes the keys & values for map
+	 *
+	 * @param Name the name of the property to serialize
+	 * @param Map the map to serialize
+	 */
+	JSON_API virtual void SerializeMap(FStringView Name, FJsonSerializableKeyValueMap& Map) override;
+
+	/**
+	 * Serializes the keys & values for map
+	 *
+	 * @param Name the name of the property to serialize
+	 * @param Map the map to serialize
+	 */
+	JSON_API virtual void SerializeMap(FStringView Name, FJsonSerializableKeyValueMapInt& Map) override;
+	
+	/**
+     * Serializes the keys & values for map
+     *
+     * @param Name the name of the property to serialize
+     * @param Map the map to serialize
+     */
+    JSON_API virtual void SerializeMap(FStringView Name, FJsonSerializableKeyValueMapArrayInt& Map) override;
+
+	/**
+	 * Serializes the keys & values for map
+	 *
+	 * @param Name the name of the property to serialize
+	 * @param Map the map to serialize
+	 */
+	JSON_API virtual void SerializeMap(FStringView Name, FJsonSerializableKeyValueMapInt64& Map) override;
+
+	/**
+	 * Serializes the keys & values for map
+	 *
+	 * @param Name the name of the property to serialize
+	 * @param Map the map to serialize
+	 */
+	JSON_API virtual void SerializeMap(FStringView Name, FJsonSerializableKeyValueMapFloat& Map) override;
+
+	JSON_API virtual void SerializeSimpleMap(FJsonSerializableKeyValueMap& Map) override;
+
+	/**
+	 * Deserializes keys and values from an object into a map, but only if the value is trivially convertable to string.
+	 *
+	 * @param Name Name of property to deserialize
+	 * @param Map The Map to fill with String values found
+	 */
+	JSON_API virtual void SerializeMapSafe(FStringView Name, FJsonSerializableKeyValueMap& Map) override;
+
+	JSON_API virtual void WriteIdentifierPrefix(FStringView Name);
+
+	JSON_API virtual void WriteRawJSONValue(FStringView Value);
+
+private:
+	/** The object that holds the parsed JSON data */
+	TSharedPtr<FJsonObject> JsonObject;
+};
+
diff --git a/Engine/Source/Runtime/Json/Public/Serialization/JsonSerializerWriter.h b/Engine/Source/Runtime/Json/Public/Serialization/JsonSerializerWriter.h
new file mode 100644
--- /dev/null
+++ b/Engine/Source/Runtime/Json/Public/Serialization/JsonSerializerWriter.h
@@ -0,0 +1,350 @@
+// Copyright Epic Games, Inc. All Rights Reserved.
+
+#pragma once
+
+#include "Policies/PrettyJsonPrintPolicy.h"
+#include "Serialization/JsonTypes.h"
+#include "Serialization/JsonSerializerBase.h"
+#include "Serialization/JsonWriter.h"
+
+/**
+ * Implements the abstract serializer interface hiding the underlying writer object
+ */
+template <class CharType = TCHAR, class PrintPolicy = TPrettyJsonPrintPolicy<CharType> >
+class FJsonSerializerWriter :
+	public FJsonSerializerBase
+{
+	/** The object to write the JSON output to */
+	TSharedRef<TJsonWriter<CharType, PrintPolicy> > JsonWriter;
+
+public:
+
+	/**
+	 * Initializes the writer object
+	 *
+	 * @param InJsonWriter the object to write the JSON output to
+	 */
+	FJsonSerializerWriter(TSharedRef<TJsonWriter<CharType, PrintPolicy> > InJsonWriter) :
+		JsonWriter(InJsonWriter)
+	{
+	}
+
+	virtual ~FJsonSerializerWriter()
+	{
+	}
+
+	/** Is the JSON being read from */
+	virtual bool IsLoading() const override { return false; }
+	/** Is the JSON being written to */
+	virtual bool IsSaving() const override { return true; }
+	/** Access to the root object */
+	virtual TSharedPtr<FJsonObject> GetObject() override { return TSharedPtr<FJsonObject>(); }
+
+	/**
+	 * Starts a new object "{"
+	 */
+	virtual void StartObject() override
+	{
+		JsonWriter->WriteObjectStart();
+	}
+
+	/**
+	 * Starts a new object "{"
+	 */
+	virtual void StartObject(FStringView Name) override
+	{
+		JsonWriter->WriteObjectStart(Name);
+	}
+	/**
+	 * Completes the definition of an object "}"
+	 */
+	virtual void EndObject() override
+	{
+		JsonWriter->WriteObjectEnd();
+	}
+
+	virtual void StartArray() override
+	{
+		JsonWriter->WriteArrayStart();
+	}
+
+	virtual void StartArray(FStringView Name) override
+	{
+		JsonWriter->WriteArrayStart(Name);
+	}
+
+	virtual void EndArray() override
+	{
+		JsonWriter->WriteArrayEnd();
+	}
+	/**
+	 * Writes the field name and the corresponding value to the JSON data
+	 *
+	 * @param Name the field name to write out
+	 * @param Value the value to write out
+	 */
+	virtual void Serialize(FStringView Name, int32& Value) override
+	{
+		JsonWriter->WriteValue(Name, Value);
+	}
+	/**
+	 * Writes the field name and the corresponding value to the JSON data
+	 *
+	 * @param Name the field name to write out
+	 * @param Value the value to write out
+	 */
+	virtual void Serialize(FStringView Name, uint32& Value) override
+	{
+		JsonWriter->WriteValue(Name, static_cast<int64>(Value));
+	}
+	/**
+	 * Writes the field name and the corresponding value to the JSON data
+	 *
+	 * @param Name the field name to write out
+	 * @param Value the value to write out
+	 */
+	virtual void Serialize(FStringView Name, int64& Value) override
+	{
+		JsonWriter->WriteValue(Name, Value);
+	}
+	/**
+	 * Writes the field name and the corresponding value to the JSON data
+	 *
+	 * @param Name the field name to write out
+	 * @param Value the value to write out
+	 */
+	virtual void Serialize(FStringView Name, bool& Value) override
+	{
+		JsonWriter->WriteValue(Name, Value);
+	}
+	/**
+	 * Writes the field name and the corresponding value to the JSON data
+	 *
+	 * @param Name the field name to write out
+	 * @param Value the value to write out
+	 */
+	virtual void Serialize(FStringView Name, FString& Value) override
+	{
+		JsonWriter->WriteValue(Name, Value);
+	}
+	/**
+	 * Writes the field name and the corresponding value to the JSON data
+	 *
+	 * @param Name the field name to write out
+	 * @param Value the value to write out
+	 */
+	virtual void Serialize(FStringView Name, FText& Value) override
+	{
+		JsonWriter->WriteValue(Name, Value.ToString());
+	}
+	/**
+	 * Writes the field name and the corresponding value to the JSON data
+	 *
+	 * @param Name the field name to write out
+	 * @param Value the value to write out
+	 */
+	virtual void Serialize(FStringView Name, float& Value) override
+	{
+		JsonWriter->WriteValue(Name, Value);
+	}
+	/**
+	* Writes the field name and the corresponding value to the JSON data
+	*
+	* @param Name the field name to write out
+	* @param Value the value to write out
+	*/
+	virtual void Serialize(FStringView Name, double& Value) override
+	{
+		JsonWriter->WriteValue(Name, Value);
+	}
+	/**
+	* Writes the field name and the corresponding value to the JSON data
+	*
+	* @param Name the field name to write out
+	* @param Value the value to write out
+	*/
+	virtual void Serialize(FStringView Name, FDateTime& Value) override
+	{
+		if (Value.GetTicks() > 0)
+		{
+			JsonWriter->WriteValue(Name, Value.ToIso8601());
+		}
+	}
+	/**
+	 * Serializes an array of values
+	 *
+	 * @param Name the name of the property to serialize
+	 * @param Array the array to serialize
+	 */
+	virtual void SerializeArray(FJsonSerializableArray& Array) override
+	{
+		JsonWriter->WriteArrayStart();
+		// Iterate all of values
+		for (FJsonSerializableArray::TIterator ArrayIt(Array); ArrayIt; ++ArrayIt)
+		{
+			JsonWriter->WriteValue(*ArrayIt);
+		}
+		JsonWriter->WriteArrayEnd();
+	}
+	/**
+	 * Serializes an array of values with an identifier
+	 *
+	 * @param Name the name of the property to serialize
+	 * @param Array the array to serialize
+	 */
+	virtual void SerializeArray(FStringView Name, FJsonSerializableArray& Array) override
+	{
+		JsonWriter->WriteArrayStart(Name);
+		// Iterate all of values
+		for (FJsonSerializableArray::ElementType& Item :  Array)
+		{
+			JsonWriter->WriteValue(Item);
+		}
+		JsonWriter->WriteArrayEnd();
+	}
+	/**
+	 * Serializes an array of values with an identifier
+	 *
+	 * @param Name the name of the property to serialize
+	 * @param Array the array to serialize
+	 */
+	virtual void SerializeArray(FStringView Name, FJsonSerializableArrayInt& Array) override
+	{
+		JsonWriter->WriteArrayStart(Name);
+		// Iterate all of values
+		for (FJsonSerializableArrayInt::ElementType& Item : Array)
+		{
+			JsonWriter->WriteValue(Item);
+		}
+		JsonWriter->WriteArrayEnd();
+	}
+
+	/**
+	 * Serializes an array of values with an identifier
+	 *
+	 * @param Name the name of the property to serialize
+	 * @param Array the array to serialize
+	 */
+	virtual void SerializeArray(FStringView Name, FJsonSerializableArrayFloat& Array) override
+	{
+		JsonWriter->WriteArrayStart(Name);
+		// Iterate all of values
+		for (FJsonSerializableArrayFloat::ElementType& Item : Array)
+		{
+			JsonWriter->WriteValue(Item);
+		}
+		JsonWriter->WriteArrayEnd();
+	}
+
+	/**
+	 * Serializes the keys & values for map
+	 *
+	 * @param Name the name of the property to serialize
+	 * @param Map the map to serialize
+	 */
+	virtual void SerializeMap(FStringView Name, FJsonSerializableKeyValueMap& Map) override
+	{
+		JsonWriter->WriteObjectStart(Name);
+		// Iterate all of the keys and their values
+		for (FJsonSerializableKeyValueMap::ElementType& Pair : Map)
+		{
+			Serialize(*Pair.Key, Pair.Value);
+		}
+		JsonWriter->WriteObjectEnd();
+	}
+
+	/**
+	 * Serializes the keys & values for map
+	 *
+	 * @param Name the name of the property to serialize
+	 * @param Map the map to serialize
+	 */
+	virtual void SerializeMap(FStringView Name, FJsonSerializableKeyValueMapInt& Map) override
+	{
+		JsonWriter->WriteObjectStart(Name);
+		// Iterate all of the keys and their values
+		for (FJsonSerializableKeyValueMapInt::ElementType& Pair : Map)
+		{
+			Serialize(*Pair.Key, Pair.Value);
+		}
+		JsonWriter->WriteObjectEnd();
+	}
+	/**
+	 * Serializes the keys & values for map
+	 *
+	 * @param Name the name of the property to serialize
+	 * @param Map the map to serialize
+	 */
+	virtual void SerializeMap(FStringView Name, FJsonSerializableKeyValueMapArrayInt& Map) override
+	{
+		JsonWriter->WriteObjectStart(Name);
+		// Iterate all of the keys and their values
+		for (FJsonSerializableKeyValueMapArrayInt::ElementType& Pair : Map)
+		{
+			SerializeArray(*Pair.Key, Pair.Value);
+		}
+		JsonWriter->WriteObjectEnd();
+	}
+
+	/**
+	 * Serializes the keys & values for map
+	 *
+	 * @param Name the name of the property to serialize
+	 * @param Map the map to serialize
+	 */
+	virtual void SerializeMap(FStringView Name, FJsonSerializableKeyValueMapInt64& Map) override
+	{
+		JsonWriter->WriteObjectStart(Name);
+		// Iterate all of the keys and their values
+		for (FJsonSerializableKeyValueMapInt64::ElementType& Pair : Map)
+		{
+			Serialize(*Pair.Key, Pair.Value);
+		}
+		JsonWriter->WriteObjectEnd();
+	}
+
+	/**
+	 * Serializes the keys & values for map
+	 *
+	 * @param Name the name of the property to serialize
+	 * @param Map the map to serialize
+	 */
+	virtual void SerializeMap(FStringView Name, FJsonSerializableKeyValueMapFloat& Map) override
+	{
+		JsonWriter->WriteObjectStart(Name);
+		// Iterate all of the keys and their values
+		for (FJsonSerializableKeyValueMapFloat::ElementType& Pair : Map)
+		{
+			Serialize(*Pair.Key, Pair.Value);
+		}
+		JsonWriter->WriteObjectEnd();
+	}
+
+	virtual void SerializeSimpleMap(FJsonSerializableKeyValueMap& Map) override
+	{
+		// writing does nothing here, this is meant to read in all data from a json object 
+		// writing is explicitly handled per key/type
+	}
+
+	/**
+	 * Serializes keys and values from an object into a map.
+	 *
+	 * @param Name Name of property to serialize
+	 * @param Map The Map to copy String values from
+	 */
+	virtual void SerializeMapSafe(FStringView Name, FJsonSerializableKeyValueMap& Map)
+	{
+		SerializeMap(Name, Map);
+	}
+
+	virtual void WriteIdentifierPrefix(FStringView Name)
+	{
+		JsonWriter->WriteIdentifierPrefix(Name);
+	}
+
+	virtual void WriteRawJSONValue(FStringView Value)
+	{
+		JsonWriter->WriteRawJSONValue(Value);
+	}
+};
+
diff --git a/Engine/Source/Runtime/Json/Public/Serialization/JsonTypes.h b/Engine/Source/Runtime/Json/Public/Serialization/JsonTypes.h
--- a/Engine/Source/Runtime/Json/Public/Serialization/JsonTypes.h
+++ b/Engine/Source/Runtime/Json/Public/Serialization/JsonTypes.h
@@ -67,3 +67,16 @@ enum class EJsonNotation
 	Null,
 	Error
 };
+
+/** Array of data */
+using FJsonSerializableArray = TArray<FString>;
+using FJsonSerializableArrayInt = TArray<int32>;
+using FJsonSerializableArrayFloat = TArray<float>;
+
+/** Maps a key to a value */
+using FJsonSerializableKeyValueMap = TMap<FString, FString>;
+using FJsonSerializableKeyValueMapInt = TMap<FString, int32>;
+using FJsonSerializableKeyValueMapInt64 = TMap<FString, int64>;
+using FJsonSerializableKeyValueMapFloat = TMap<FString, float>;
+using FJsonSerializableKeyValueMapArrayInt = TMap<FString, FJsonSerializableArrayInt>;
+
