Make Json module reference the RapidJson library provided in ThirdParty folder.

diff --git a/Engine/Source/Runtime/Json/Json.Build.cs b/Engine/Source/Runtime/Json/Json.Build.cs
--- a/Engine/Source/Runtime/Json/Json.Build.cs
+++ b/Engine/Source/Runtime/Json/Json.Build.cs
@@ -10,6 +10,7 @@ namespace UnrealBuildTool.Rules
 				new string[]
 				{
 					"Core",
+					"RapidJSON",
 				}
 			); 
 
diff --git a/Engine/Source/Runtime/Json/Public/RapidJsonIncludes.h b/Engine/Source/Runtime/Json/Public/RapidJsonIncludes.h
new file mode 100644
--- /dev/null
+++ b/Engine/Source/Runtime/Json/Public/RapidJsonIncludes.h
@@ -0,0 +1,37 @@
+// Copyright Epic Games, Inc. All Rights Reserved.
+
+#pragma once
+
+#if defined(_MSC_VER) && USING_CODE_ANALYSIS
+#pragma warning(push)
+#pragma warning(disable : 6282) // Incorrect operator: Assignment of constant in Boolean context. Consider using '==' instead
+#pragma warning(disable : 6313) // Incorrect operator: Zero-valued flag cannot be tested with bitwise-and. Use an equality test to check for zero-valued flags
+#endif
+
+THIRD_PARTY_INCLUDES_START
+#ifndef RAPIDJSON_ASSERT
+#define RAPIDJSON_ASSERT(x) check(x)
+#endif
+
+#ifndef RAPIDJSON_ERROR_CHARTYPE
+#define RAPIDJSON_ERROR_CHARTYPE TCHAR
+#endif
+
+#ifndef RAPIDJSON_ERROR_STRING
+#define RAPIDJSON_ERROR_STRING(x) TEXT(x)
+#endif
+
+#include "rapidjson/rapidjson.h"
+#include "rapidjson/document.h"
+#include "rapidjson/encodings.h"
+#include "rapidjson/writer.h"
+#include "rapidjson/stringbuffer.h"
+#include "rapidjson/schema.h"
+// just for non-user facing logging purposes use the English error descriptions
+#include "rapidjson/error/en.h"
+THIRD_PARTY_INCLUDES_END
+
+#if defined(_MSC_VER) && USING_CODE_ANALYSIS
+#pragma warning(pop)
+#endif
+
