Rename the template method FromJson to function overloads of FromJsonStringView, to not expose more types before actually supporting them with RapidJson.

diff --git a/Engine/Source/Runtime/Json/Private/Serialization/JsonSerializable.cpp b/Engine/Source/Runtime/Json/Private/Serialization/JsonSerializable.cpp
--- a/Engine/Source/Runtime/Json/Private/Serialization/JsonSerializable.cpp
+++ b/Engine/Source/Runtime/Json/Private/Serialization/JsonSerializable.cpp
@@ -49,7 +49,7 @@ void FJsonSerializable::ToJson(TSharedRef< TJsonWriter< TCHAR, TCondensedJsonPri
 
 bool FJsonSerializable::FromJson(const FString& Json)
 {
-	return FromJson(CopyTemp(Json));
+	return FromJsonStringView(FStringView(Json));
 }
 
 bool FJsonSerializable::FromJson(FString&& Json)
@@ -67,6 +67,37 @@ bool FJsonSerializable::FromJson(FString&& Json)
 	return false;
 }
 
+namespace UE::JsonSerializable::Private
+{
+
+template<typename CharType>
+bool FromJsonStringView(FJsonSerializable* Serializable, TStringView<CharType> JsonStringView)
+{
+	TSharedPtr<FJsonObject> JsonObject;
+	TSharedRef<TJsonReader<CharType> > JsonReader = TJsonReaderFactory<CharType>::CreateFromView(JsonStringView);
+	if (FJsonSerializer::Deserialize(JsonReader,JsonObject) &&
+		JsonObject.IsValid())
+	{
+		FJsonSerializerReader Serializer(JsonObject);
+		Serializable->Serialize(Serializer, false);
+		return true;
+	}
+	UE_LOG(LogJson, Warning, TEXT("Failed to parse Json from a string: %s"), *JsonReader->GetErrorMessage());
+	return false;
+}
+
+}
+
+bool FJsonSerializable::FromJsonStringView(FUtf8StringView JsonStringView)
+{
+	return UE::JsonSerializable::Private::FromJsonStringView(this, JsonStringView);
+}
+
+bool FJsonSerializable::FromJsonStringView(FWideStringView JsonStringView)
+{
+	return UE::JsonSerializable::Private::FromJsonStringView(this, JsonStringView);
+}
+
 bool FJsonSerializable::FromJson(TSharedPtr<FJsonObject> JsonObject)
 {
 	if (JsonObject.IsValid())
diff --git a/Engine/Source/Runtime/Json/Public/Serialization/JsonSerializable.h b/Engine/Source/Runtime/Json/Public/Serialization/JsonSerializable.h
--- a/Engine/Source/Runtime/Json/Public/Serialization/JsonSerializable.h
+++ b/Engine/Source/Runtime/Json/Public/Serialization/JsonSerializable.h
@@ -48,21 +48,20 @@ struct FJsonSerializable
 	 */
 	JSON_API virtual bool FromJson(FString&& Json);
 
-	template <class CharType = TCHAR>
-	bool FromJson(TStringView<CharType> Json)
-	{
-		TSharedPtr<FJsonObject> JsonObject;
-		TSharedRef<TJsonReader<CharType> > JsonReader = TJsonReaderFactory<CharType>::CreateFromView(Json);
-		if (FJsonSerializer::Deserialize(JsonReader,JsonObject) &&
-			JsonObject.IsValid())
-		{
-			FJsonSerializerReader Serializer(JsonObject);
-			Serialize(Serializer, false);
-			return true;
-		}
-		UE_LOG(LogJson, Warning, TEXT("Failed to parse Json from a string: %s"), *JsonReader->GetErrorMessage());
-		return false;
-	}
+	/**
+	 * Serializes the contents of a JSON string into this object using FUtf8StringView
+	 *
+	 * @param JsonStringView the JSON data to serialize from
+	 */
+	JSON_API bool FromJsonStringView(FUtf8StringView JsonStringView);
+
+	/**
+	 * Serializes the contents of a JSON string into this object using FWideStringView
+	 *
+	 * @param JsonStringView the JSON data to serialize from
+	 */
+	JSON_API bool FromJsonStringView(FWideStringView JsonStringView);
+ 
 
 	JSON_API virtual bool FromJson(TSharedPtr<FJsonObject> JsonObject);
 
