Implemented `JSON_SERIALIZE_MAP_ARRAY_SERIALIZABLE`: Serializes a `TMap<FString, TArray<SerializableT>>`

diff --git a/Engine/Plugins/Online/OnlineSubsystem/Source/Public/OnlineJsonSerializer.h b/Engine/Plugins/Online/OnlineSubsystem/Source/Public/OnlineJsonSerializer.h
--- a/Engine/Plugins/Online/OnlineSubsystem/Source/Public/OnlineJsonSerializer.h
+++ b/Engine/Plugins/Online/OnlineSubsystem/Source/Public/OnlineJsonSerializer.h
@@ -18,6 +18,7 @@ class FVariantData;
 #define ONLINE_JSON_SERIALIZE_ARRAY_SERIALIZABLE(JsonName, JsonArray, ElementType)					JSON_SERIALIZE_ARRAY_SERIALIZABLE(JsonName, JsonArray, ElementType)
 #define ONLINE_JSON_SERIALIZE_OPTIONAL_ARRAY_SERIALIZABLE(JsonName, OptionalJsonArray, ElementType)	JSON_SERIALIZE_OPTIONAL_ARRAY_SERIALIZABLE(JsonName, OptionalJsonArray, ElementType)
 #define ONLINE_JSON_SERIALIZE_MAP_SERIALIZABLE(JsonName, JsonMap, ElementType)						JSON_SERIALIZE_MAP_SERIALIZABLE(JsonName, JsonMap, ElementType) 
+#define ONLINE_JSON_SERIALIZE_MAP_ARRAY_SERIALIZABLE(JsonName, JsonMap, ElementType)				JSON_SERIALIZE_MAP_ARRAY_SERIALIZABLE(JsonName, JsonMap, ElementType) 
 #define ONLINE_JSON_SERIALIZE_OBJECT_SERIALIZABLE(JsonName, JsonSerializableObject)					JSON_SERIALIZE_OBJECT_SERIALIZABLE(JsonName, JsonSerializableObject) 
 #define ONLINE_JSON_SERIALIZE_DATETIME_UNIX_TIMESTAMP(JsonName, JsonDateTime)						JSON_SERIALIZE_DATETIME_UNIX_TIMESTAMP(JsonName, JsonDateTime)
 #define ONLINE_JSON_SERIALIZE_ENUM(JsonName, JsonEnum)												JSON_SERIALIZE_ENUM(JsonName, JsonEnum)
diff --git a/Engine/Source/Runtime/Json/Public/Serialization/JsonSerializerMacros.h b/Engine/Source/Runtime/Json/Public/Serialization/JsonSerializerMacros.h
--- a/Engine/Source/Runtime/Json/Public/Serialization/JsonSerializerMacros.h
+++ b/Engine/Source/Runtime/Json/Public/Serialization/JsonSerializerMacros.h
@@ -217,6 +217,38 @@
 			Serializer.EndObject(); \
 		}
 
+#define JSON_SERIALIZE_MAP_ARRAY_SERIALIZABLE(JsonName, JsonMap, ElementType) \
+		if (Serializer.IsLoading()) \
+		{ \
+			if (Serializer.GetObject()->HasTypedField<EJson::Object>(TEXTVIEW(JsonName))) \
+			{ \
+				TSharedPtr<FJsonObject> JsonObj = Serializer.GetObject()->GetObjectField(TEXTVIEW(JsonName)); \
+				for (auto MapIt = JsonObj->Values.CreateConstIterator(); MapIt; ++MapIt) \
+				{ \
+					TArray<ElementType> NewEntry; \
+					for (auto ArrayIt : MapIt.Value()->AsArray()) \
+					{ \
+						NewEntry.Add_GetRef(ElementType()).FromJson(ArrayIt->AsObject()); \
+					} \
+					JsonMap.Add(MapIt.Key(), NewEntry); \
+				} \
+			} \
+		} \
+		else \
+		{ \
+			Serializer.StartObject(TEXTVIEW(JsonName)); \
+			for (auto It = JsonMap.CreateIterator(); It; ++It) \
+			{ \
+				Serializer.StartArray(It.Key()); \
+				for (auto ArrayEntry : It.Value()) \
+				{ \
+					ArrayEntry.Serialize(Serializer, false); \
+				} \
+				Serializer.EndArray(); \
+			} \
+			Serializer.EndObject(); \
+		}
+
 #define JSON_SERIALIZE_OBJECT_SERIALIZABLE(JsonName, JsonSerializableObject) \
 		/* Process the JsonName field differently because it is an object */ \
 		if (Serializer.IsLoading()) \
