Implemented `TPrettyJsonBag<int32>` to use the `TPrettyJsonPrintPolicy` trait instead of `TCondensedJsonPrintPolicy`

diff --git a/Engine/Source/Runtime/Json/Private/Serialization/JsonDataBag.cpp b/Engine/Source/Runtime/Json/Private/Serialization/JsonDataBag.cpp
--- a/Engine/Source/Runtime/Json/Private/Serialization/JsonDataBag.cpp
+++ b/Engine/Source/Runtime/Json/Private/Serialization/JsonDataBag.cpp
@@ -49,8 +49,17 @@ void FJsonDataBag::Serialize(FJsonSerializerBase& Serializer, bool bFlatObject)
 						{
 							// if we have an array, serialize to string and write raw
 							FString JsonStr;
-							auto Writer = TJsonWriterFactory<TCHAR, TCondensedJsonPrintPolicy<TCHAR>>::Create(&JsonStr);
-							FJsonSerializer::Serialize(JsonValue->AsArray(), Writer);
+							if (NumPrintIndents.IsSet())
+							{
+								auto Writer = TJsonWriterFactory<TCHAR, TPrettyJsonPrintPolicy<TCHAR>>::Create(&JsonStr, *NumPrintIndents);
+								FJsonSerializer::Serialize(JsonValue->AsArray(), Writer);
+							}
+							else
+							{
+								auto Writer = TJsonWriterFactory<TCHAR, TCondensedJsonPrintPolicy<TCHAR>>::Create(&JsonStr);
+								FJsonSerializer::Serialize(JsonValue->AsArray(), Writer);
+							}
+							
 							Serializer.WriteIdentifierPrefix(*It.Key);
 							Serializer.WriteRawJSONValue(*JsonStr);
 							break;
@@ -59,8 +68,18 @@ void FJsonDataBag::Serialize(FJsonSerializerBase& Serializer, bool bFlatObject)
 						{
 							// if we have an object, serialize to string and write raw
 							FString JsonStr;
-							auto Writer = TJsonWriterFactory<TCHAR, TCondensedJsonPrintPolicy<TCHAR>>::Create(&JsonStr);
-							FJsonSerializer::Serialize(JsonValue->AsObject().ToSharedRef(), Writer);
+
+							if (NumPrintIndents.IsSet())
+							{
+								auto Writer = TJsonWriterFactory<TCHAR, TPrettyJsonPrintPolicy<TCHAR>>::Create(&JsonStr, *NumPrintIndents);
+								FJsonSerializer::Serialize(JsonValue->AsObject().ToSharedRef(), Writer);
+							}
+							else
+							{
+								auto Writer = TJsonWriterFactory<TCHAR, TCondensedJsonPrintPolicy<TCHAR>>::Create(&JsonStr);
+								FJsonSerializer::Serialize(JsonValue->AsObject().ToSharedRef(), Writer);
+							}
+
 							// too bad there's no JsonObject serialization method on FJsonSerializerBase directly :-/
 							Serializer.WriteIdentifierPrefix(*It.Key);
 							Serializer.WriteRawJSONValue(*JsonStr);
diff --git a/Engine/Source/Runtime/Json/Public/Serialization/JsonDataBag.h b/Engine/Source/Runtime/Json/Public/Serialization/JsonDataBag.h
--- a/Engine/Source/Runtime/Json/Public/Serialization/JsonDataBag.h
+++ b/Engine/Source/Runtime/Json/Public/Serialization/JsonDataBag.h
@@ -3,6 +3,7 @@
 #pragma once
 
 #include "Serialization/JsonSerializable.h"
+#include "Misc/Optional.h"
 
 class FJsonSerializerReader;
 
@@ -28,7 +29,19 @@ struct FJsonDataBag : public FJsonSerializable
 	}
 
 	JSON_API void SetFieldJson(const FString& Key, const TSharedPtr<FJsonValue>& Value);
-
 public:
 	TSharedPtr<FJsonObject> JsonObject;
+
+	/* If set, will use TPrettyJsonPrintPolicy */
+	TOptional<int32> NumPrintIndents;
+};
+
+template<int32 PrintIndentCount>
+struct TPrettyJsonBag : public FJsonDataBag
+{
+	TPrettyJsonBag()
+		: FJsonDataBag()
+	{
+		NumPrintIndents.Emplace(PrintIndentCount);
+	}
 };
