Implemented JSON serialization whose root object is an Array of serializable objects, i.e. the first character of the json file is a '['

diff --git a/Engine/Source/Runtime/Json/Public/Serialization/JsonSerializable.h b/Engine/Source/Runtime/Json/Public/Serialization/JsonSerializable.h
--- a/Engine/Source/Runtime/Json/Public/Serialization/JsonSerializable.h
+++ b/Engine/Source/Runtime/Json/Public/Serialization/JsonSerializable.h
@@ -7,6 +7,8 @@
 #include "Serialization/JsonWriter.h"
 #include "Serialization/JsonSerializer.h"
 #include "Serialization/JsonSerializerReader.h"
+#include "JsonSerializerWriter.h"
+#include "Misc/TVariant.h"
 
 /**
  * Base class for a JSON serializable object
@@ -61,7 +63,6 @@ struct FJsonSerializable
 	 * @param JsonStringView the JSON data to serialize from
 	 */
 	JSON_API bool FromJsonStringView(FWideStringView JsonStringView);
- 
 
 	JSON_API virtual bool FromJson(TSharedPtr<FJsonObject> JsonObject);
 
@@ -73,3 +74,160 @@ struct FJsonSerializable
 	 */
 	JSON_API virtual void Serialize(FJsonSerializerBase& Serializer, bool bFlatObject) = 0;
 };
+
+namespace UE::JsonArray
+{
+	namespace Private
+	{
+		template<typename T, typename CharType>
+		inline bool FromJson(TArray<T>& OutArray, TStringView<CharType> JsonString)
+		{
+			OutArray.Reset();
+
+			TArray<TSharedPtr<FJsonValue>> ArrayValues;
+			TSharedRef<TJsonReader<CharType>> JsonReader = TJsonReaderFactory<CharType>::CreateFromView(JsonString);
+			if (FJsonSerializer::Deserialize(JsonReader, ArrayValues))
+			{
+				for (const TSharedPtr<FJsonValue>& Value : ArrayValues)
+				{
+					TSharedPtr<FJsonObject>* ArrayEntry;
+					if (Value.IsValid() && Value->TryGetObject(ArrayEntry))
+					{
+						if (ArrayEntry && ArrayEntry->IsValid())
+						{
+							FJsonSerializerReader Serializer(*ArrayEntry);
+							OutArray.Add_GetRef(T()).Serialize(Serializer, false);
+						}
+						else
+						{
+							UE_LOG(LogJson, Error, TEXT("Failed to parse Json from array"));
+							return false;
+						}
+					}
+				}
+
+				return true;
+			}
+
+			return false;
+		}
+
+		using ReturnStringArgs = TTuple<FString* /*OutValue*/, bool /*bPrettyPrint*/>;
+
+		using PrettyWriter = TSharedRef<TJsonWriter<>>;
+		using CondensedWriter = TSharedRef<TJsonWriter<TCHAR, TCondensedJsonPrintPolicy<TCHAR>>>;
+		using WriterVariants = TVariant<PrettyWriter, CondensedWriter>;
+
+		using ToJsonVariantArgs = TVariant<ReturnStringArgs, WriterVariants>;
+
+		using PrettySerializer = FJsonSerializerWriter<>;
+		using CondensedSerializer = FJsonSerializerWriter<TCHAR, TCondensedJsonPrintPolicy<TCHAR>>;
+
+		template<typename T>
+		inline void ToJson(TArray<T>& InArray, const ToJsonVariantArgs& InArgs)
+		{
+			using PrettySerializerAndWriter = TTuple<PrettySerializer, PrettyWriter>;
+			using CondensedSerializerAndWriter = TTuple<CondensedSerializer, CondensedWriter>;
+
+			using SerializerVariant = TVariant<PrettySerializerAndWriter, CondensedSerializerAndWriter>;
+
+			SerializerVariant SerializerToUse = ::Visit([](auto& StoredValue)
+				{
+					using StoredValueType = std::decay_t<decltype(StoredValue)>;
+					if constexpr (std::is_same_v<StoredValueType, ReturnStringArgs>)
+					{
+						if (StoredValue.Get<1>())
+						{
+							PrettyWriter NewWriter = TJsonWriterFactory<>::Create(StoredValue.Get<0>());
+							return SerializerVariant(TInPlaceType<PrettySerializerAndWriter>(), PrettySerializer(NewWriter), NewWriter);
+						}
+						else
+						{
+							CondensedWriter NewWriter = TJsonWriterFactory<TCHAR, TCondensedJsonPrintPolicy<TCHAR>>::Create(StoredValue.Get<0>());
+							return SerializerVariant(TInPlaceType<CondensedSerializerAndWriter>(), CondensedSerializer(NewWriter), NewWriter);
+						}
+					}
+					else
+					{
+						return ::Visit([](auto& StoredWriter)
+							{
+								using StoredWriterType = std::decay_t<decltype(StoredWriter)>;
+								if constexpr (std::is_same_v<StoredWriterType, PrettyWriter>)
+								{
+									return SerializerVariant(TInPlaceType<PrettySerializerAndWriter>(), PrettySerializer(StoredWriter), StoredWriter);
+								}
+								else if constexpr (std::is_same_v<StoredWriterType, CondensedWriter>)
+								{
+									return SerializerVariant(TInPlaceType<CondensedSerializerAndWriter>(), CondensedSerializer(StoredWriter), StoredWriter);
+								}
+							}, StoredValue);
+					}
+				}, InArgs);
+
+
+			const bool bCloseWriter = InArgs.IsType<ReturnStringArgs>();
+
+			::Visit([bCloseWriter, &InArray](auto& StoredSerializer)
+				{
+					StoredSerializer.Get<0>().StartArray();
+
+					for (T& ArrayEntry : InArray)
+					{
+						ArrayEntry.Serialize(StoredSerializer.Get<0>(), false);
+					}
+
+					StoredSerializer.Get<0>().EndArray();
+
+					if (bCloseWriter)
+					{
+						StoredSerializer.Get<1>()->Close();
+					}
+				}, SerializerToUse);
+		}
+	}
+
+	template<typename T>
+	static bool FromJson(TArray<T>& OutArray, const FString& JsonString)
+	{
+		return Private::FromJson(OutArray, FStringView(JsonString));
+	}
+
+	template<typename T>
+	static bool FromJson(TArray<T>& OutArray, FString&& JsonString)
+	{
+		return Private::FromJson(OutArray, FStringView(MoveTemp(JsonString)));
+	}
+
+	template<typename T>
+	static bool FromJson(TArray<T>& OutArray, FUtf8StringView JsonStringView)
+	{
+		return Private::FromJson(OutArray, JsonStringView);
+	}
+
+	template<typename T>
+	static bool FromJson(TArray<T>& OutArray, FWideStringView JsonStringView)
+	{
+		return Private::FromJson(OutArray, JsonStringView);
+	}
+
+	/* non-const due to T::Serialize being a non-const function */
+	template<typename T>
+	static const FString ToJson(TArray<T>& InArray, const bool bPrettyPrint = true)
+	{
+		FString JsonStr;
+		Private::ToJson(InArray, Private::ToJsonVariantArgs(TInPlaceType<Private::ReturnStringArgs>(), &JsonStr, bPrettyPrint));
+		return JsonStr;
+	}
+
+	template<typename T>
+	static void ToJson(TArray<T>& InArray, Private::PrettyWriter& JsonWriter)
+	{
+		Private::ToJson(InArray, Private::ToJsonVariantArgs(TInPlaceType<Private::WriterVariants>(), Private::WriterVariants(TInPlaceType<Private::PrettyWriter>(), JsonWriter)));
+	}
+
+	template<typename T>
+	static void ToJson(TArray<T>& InArray, Private::CondensedWriter& JsonWriter)
+	{
+		Private::ToJson(InArray, Private::ToJsonVariantArgs(TInPlaceType<Private::WriterVariants>(), Private::WriterVariants(TInPlaceType<Private::CondensedWriter>(), JsonWriter)));
+	}
+}
\ No newline at end of file
