Compilation Fix

diff --git a/Engine/Source/Runtime/Json/Public/Serialization/JsonSerializable.h b/Engine/Source/Runtime/Json/Public/Serialization/JsonSerializable.h
--- a/Engine/Source/Runtime/Json/Public/Serialization/JsonSerializable.h
+++ b/Engine/Source/Runtime/Json/Public/Serialization/JsonSerializable.h
@@ -136,14 +136,14 @@ namespace UE::JsonArray
 					using StoredValueType = std::decay_t<decltype(StoredValue)>;
 					if constexpr (std::is_same_v<StoredValueType, ReturnStringArgs>)
 					{
-						if (StoredValue.Get<1>())
+						if (StoredValue.template Get<1>())
 						{
-							PrettyWriter NewWriter = TJsonWriterFactory<>::Create(StoredValue.Get<0>());
+							PrettyWriter NewWriter = TJsonWriterFactory<>::Create(StoredValue.template Get<0>());
 							return SerializerVariant(TInPlaceType<PrettySerializerAndWriter>(), PrettySerializer(NewWriter), NewWriter);
 						}
 						else
 						{
-							CondensedWriter NewWriter = TJsonWriterFactory<TCHAR, TCondensedJsonPrintPolicy<TCHAR>>::Create(StoredValue.Get<0>());
+							CondensedWriter NewWriter = TJsonWriterFactory<TCHAR, TCondensedJsonPrintPolicy<TCHAR>>::Create(StoredValue.template Get<0>());
 							return SerializerVariant(TInPlaceType<CondensedSerializerAndWriter>(), CondensedSerializer(NewWriter), NewWriter);
 						}
 					}
@@ -169,18 +169,18 @@ namespace UE::JsonArray
 
 			::Visit([bCloseWriter, &InArray](auto& StoredSerializer)
 				{
-					StoredSerializer.Get<0>().StartArray();
+					StoredSerializer.template Get<0>().StartArray();
 
 					for (T& ArrayEntry : InArray)
 					{
-						ArrayEntry.Serialize(StoredSerializer.Get<0>(), false);
+						ArrayEntry.Serialize(StoredSerializer.template Get<0>(), false);
 					}
 
-					StoredSerializer.Get<0>().EndArray();
+					StoredSerializer.template Get<0>().EndArray();
 
 					if (bCloseWriter)
 					{
-						StoredSerializer.Get<1>()->Close();
+						StoredSerializer.template Get<1>()->Close();
 					}
 				}, SerializerToUse);
 		}
