Convert JsonTests to LLT. Add TestFalse to Expectations.h

#tests Ran Json test via FoundationTests locally.

diff --git a/Engine/Source/Developer/LowLevelTestsRunner/Public/TestCommon/Expectations.h b/Engine/Source/Developer/LowLevelTestsRunner/Public/TestCommon/Expectations.h
--- a/Engine/Source/Developer/LowLevelTestsRunner/Public/TestCommon/Expectations.h
+++ b/Engine/Source/Developer/LowLevelTestsRunner/Public/TestCommon/Expectations.h
@@ -32,6 +32,21 @@ namespace {
 		return TestTrue(*What, Value);
 	}
 
+	bool TestFalse(const TCHAR* What, bool Value)
+	{
+		if (Value)
+		{
+			FAIL_CHECK(FString::Printf(TEXT("Expected '%s' to be false."), What));
+			return false;
+		}
+		return true;
+	}
+
+	bool TestFalse(const FString& What, bool Value)
+	{
+		return TestFalse(*What, Value);
+	}
+
 	bool TestEqual(const TCHAR* What, const int32 Actual, const int32 Expected)
 	{
 		if (Actual != Expected)
diff --git a/Engine/Source/Runtime/Json/Private/Tests/JsonTests.cpp b/Engine/Source/Runtime/Json/Private/Tests/JsonTests.cpp
--- a/Engine/Source/Runtime/Json/Private/Tests/JsonTests.cpp
+++ b/Engine/Source/Runtime/Json/Private/Tests/JsonTests.cpp
@@ -1,20 +1,19 @@
 // Copyright Epic Games, Inc. All Rights Reserved.
 
+#if WITH_TESTS
+
 #include "CoreMinimal.h"
 #include "Misc/AutomationTest.h"
 #include "Policies/CondensedJsonPrintPolicy.h"
-#include "Serialization/JsonTypes.h"
-#include "Serialization/JsonReader.h"
 #include "Policies/PrettyJsonPrintPolicy.h"
+#include "Serialization/JsonReader.h"
 #include "Serialization/JsonSerializer.h"
+#include "Serialization/JsonTypes.h"
+#include "Tests/TestHarnessAdapter.h"
 
-#if WITH_DEV_AUTOMATION_TESTS
-
-/**
- * FJsonAutomationTest
- * Simple unit test that runs Json's in-built test cases
- */
-IMPLEMENT_SIMPLE_AUTOMATION_TEST(FJsonAutomationTest, "System.Engine.FileSystem.JSON", EAutomationTestFlags::ApplicationContextMask | EAutomationTestFlags::SmokeFilter )
+#if WITH_LOW_LEVEL_TESTS
+#include "TestCommon/Expectations.h"
+#endif
 
 typedef TJsonWriterFactory< TCHAR, TCondensedJsonPrintPolicy<TCHAR> > FCondensedJsonStringWriterFactory;
 typedef TJsonWriter< TCHAR, TCondensedJsonPrintPolicy<TCHAR> > FCondensedJsonStringWriter;
@@ -22,12 +21,7 @@ typedef TJsonWriter< TCHAR, TCondensedJsonPrintPolicy<TCHAR> > FCondensedJsonStr
 typedef TJsonWriterFactory< TCHAR, TPrettyJsonPrintPolicy<TCHAR> > FPrettyJsonStringWriterFactory;
 typedef TJsonWriter< TCHAR, TPrettyJsonPrintPolicy<TCHAR> > FPrettyJsonStringWriter;
 
-/** 
- * Execute the Json test cases
- *
- * @return	true if the test was successful, false otherwise
- */
-bool FJsonAutomationTest::RunTest(const FString& Parameters)
+TEST_CASE_NAMED(FJsonAutomationTest, "System::Engine::FileSystem::JSON", "[ApplicationContextMask][SmokeFilter]")
 {
 	// Null Case
 	{
@@ -35,8 +29,8 @@ bool FJsonAutomationTest::RunTest(const FString& Parameters)
 		TSharedRef< TJsonReader<> > Reader = TJsonReaderFactory<>::Create( InputString );
 
 		TSharedPtr<FJsonObject> Object;
-		verify( FJsonSerializer::Deserialize( Reader, Object ) == false );
-		check( !Object.IsValid() );
+		REQUIRE(FJsonSerializer::Deserialize( Reader, Object ) == false);
+		REQUIRE(!Object.IsValid());
 	}
 
 	// Empty Object Case
@@ -45,13 +39,13 @@ bool FJsonAutomationTest::RunTest(const FString& Parameters)
 		TSharedRef< TJsonReader<> > Reader = TJsonReaderFactory<>::Create( InputString );
 
 		TSharedPtr<FJsonObject> Object;
-		verify( FJsonSerializer::Deserialize( Reader, Object ) );
-		check( Object.IsValid() );
+		REQUIRE(FJsonSerializer::Deserialize( Reader, Object ));
+		REQUIRE(Object.IsValid());
 
 		FString OutputString;
 		TSharedRef< FCondensedJsonStringWriter > Writer = FCondensedJsonStringWriterFactory::Create( &OutputString );
-		verify( FJsonSerializer::Serialize( Object.ToSharedRef(), Writer ) );
-		check( InputString == OutputString );
+		REQUIRE(FJsonSerializer::Serialize( Object.ToSharedRef(), Writer ));
+		REQUIRE(InputString == OutputString);
 	}
 
 	// Empty Array Case
@@ -60,13 +54,13 @@ bool FJsonAutomationTest::RunTest(const FString& Parameters)
 		TSharedRef< TJsonReader<> > Reader = TJsonReaderFactory<>::Create( InputString );
 
 		TArray< TSharedPtr<FJsonValue> > Array;
-		verify( FJsonSerializer::Deserialize( Reader, Array ) );
-		check( Array.Num() == 0 );
+		REQUIRE(FJsonSerializer::Deserialize( Reader, Array ));
+		REQUIRE(Array.Num() == 0);
 
 		FString OutputString;
 		TSharedRef< FCondensedJsonStringWriter > Writer = FCondensedJsonStringWriterFactory::Create( &OutputString );
-		verify( FJsonSerializer::Serialize( Array, Writer ) );
-		check( InputString == OutputString );
+		REQUIRE(FJsonSerializer::Serialize( Array, Writer ));
+		REQUIRE(InputString == OutputString);
 	}
 
 	// Empty Array with Empty Identifier Case
@@ -75,8 +69,8 @@ bool FJsonAutomationTest::RunTest(const FString& Parameters)
 		FString OutputString;
 		TSharedRef<FJsonValueArray> EmptyValuesArray = MakeShared<FJsonValueArray>(TArray<TSharedPtr<FJsonValue>>());
 		TSharedRef<FCondensedJsonStringWriter> JsonWriter = FCondensedJsonStringWriterFactory::Create( &OutputString );
-		verify(FJsonSerializer::Serialize(EmptyValuesArray, FString(), JsonWriter));
-		check(ExpectedString == OutputString);
+		REQUIRE(FJsonSerializer::Serialize(EmptyValuesArray, FString(), JsonWriter));
+		REQUIRE(ExpectedString == OutputString);
 	}
 
 	// Serializing Object Value with Empty Identifier Case
@@ -86,10 +80,10 @@ bool FJsonAutomationTest::RunTest(const FString& Parameters)
 		TSharedRef<FJsonValue> FooValue = MakeShared<FJsonValueString>("foo");
 		TSharedRef<FCondensedJsonStringWriter> JsonWriter = FCondensedJsonStringWriterFactory::Create(&OutputString);
 		JsonWriter->WriteObjectStart();
-		verify(FJsonSerializer::Serialize(FooValue, FString(), JsonWriter, false));
+		REQUIRE(FJsonSerializer::Serialize(FooValue, FString(), JsonWriter, false));
 		JsonWriter->WriteObjectEnd();
 		JsonWriter->Close();
-		check(ExpectedString == OutputString);
+		REQUIRE(ExpectedString == OutputString);
 	}
 
 	// Simple Array Case
@@ -107,18 +101,18 @@ bool FJsonAutomationTest::RunTest(const FString& Parameters)
 
 		TArray< TSharedPtr<FJsonValue> > Array;
 		bool bSuccessful = FJsonSerializer::Deserialize(Reader, Array);
-		check(bSuccessful);
-		check( Array.Num() == 1 );
-		check( Array[0].IsValid() );
+		REQUIRE(bSuccessful);
+		REQUIRE(Array.Num() == 1);
+		REQUIRE(Array[0].IsValid());
 
 		TSharedPtr< FJsonObject > Object = Array[0]->AsObject();
-		check( Object.IsValid() );
-		check( Object->GetStringField( TEXT("Value") ) == TEXT("Some String") );
+		REQUIRE(Object.IsValid());
+		REQUIRE(Object->GetStringField( TEXT("Value") ) == TEXT("Some String"));
 
 		FString OutputString;
 		TSharedRef< FCondensedJsonStringWriter > Writer = FCondensedJsonStringWriterFactory::Create( &OutputString );
-		verify( FJsonSerializer::Serialize( Array, Writer ) );
-		check( InputString == OutputString );
+		REQUIRE(FJsonSerializer::Serialize( Array, Writer ));
+		REQUIRE(InputString == OutputString);
 	}
 
 	// Object Array Case
@@ -143,28 +137,28 @@ bool FJsonAutomationTest::RunTest(const FString& Parameters)
 		TArray< TSharedPtr<FJsonValue> > Array;
 
 		bool bSuccessful = FJsonSerializer::Deserialize(Reader, Array);
-		check(bSuccessful);
-		check(Array.Num() == 3);
-		check(Array[0].IsValid());
-		check(Array[1].IsValid());
-		check(Array[2].IsValid());
+		REQUIRE(bSuccessful);
+		REQUIRE(Array.Num() == 3);
+		REQUIRE(Array[0].IsValid());
+		REQUIRE(Array[1].IsValid());
+		REQUIRE(Array[2].IsValid());
 
 		TSharedPtr< FJsonObject > Object = Array[0]->AsObject();
-		check(Object.IsValid());
-		check(Object->GetStringField(TEXT("Value")) == TEXT("Some String1"));
+		REQUIRE(Object.IsValid());
+		REQUIRE(Object->GetStringField(TEXT("Value")) == TEXT("Some String1"));
 
 		Object = Array[1]->AsObject();
-		check(Object.IsValid());
-		check(Object->GetStringField(TEXT("Value")) == TEXT("Some String2"));
+		REQUIRE(Object.IsValid());
+		REQUIRE(Object->GetStringField(TEXT("Value")) == TEXT("Some String2"));
 
 		Object = Array[2]->AsObject();
-		check(Object.IsValid());
-		check(Object->GetStringField(TEXT("Value")) == TEXT("Some String3"));
+		REQUIRE(Object.IsValid());
+		REQUIRE(Object->GetStringField(TEXT("Value")) == TEXT("Some String3"));
 
 		FString OutputString;
 		TSharedRef< FCondensedJsonStringWriter > Writer = FCondensedJsonStringWriterFactory::Create(&OutputString);
-		check(FJsonSerializer::Serialize(Array, Writer));
-		check(InputString == OutputString);
+		REQUIRE(FJsonSerializer::Serialize(Array, Writer));
+		REQUIRE(InputString == OutputString);
 	}
 
 	// Number Array Case
@@ -183,29 +177,29 @@ bool FJsonAutomationTest::RunTest(const FString& Parameters)
 
 		TArray< TSharedPtr<FJsonValue> > Array;
 		bool bSuccessful = FJsonSerializer::Deserialize(Reader, Array);
-		check(bSuccessful);
-		check(Array.Num() == 4);
-		check(Array[0].IsValid());
-		check(Array[1].IsValid());
-		check(Array[2].IsValid());
-		check(Array[3].IsValid());
+		REQUIRE(bSuccessful);
+		REQUIRE(Array.Num() == 4);
+		REQUIRE(Array[0].IsValid());
+		REQUIRE(Array[1].IsValid());
+		REQUIRE(Array[2].IsValid());
+		REQUIRE(Array[3].IsValid());
 
 		double Number = Array[0]->AsNumber();
-		check(Number == 10);
+		REQUIRE(Number == 10);
 
 		Number = Array[1]->AsNumber();
-		check(Number == 20);
+		REQUIRE(Number == 20);
 
 		Number = Array[2]->AsNumber();
-		check(Number == 30);
+		REQUIRE(Number == 30);
 
 		Number = Array[3]->AsNumber();
-		check(Number == 40);
+		REQUIRE(Number == 40);
 
 		FString OutputString;
 		TSharedRef< FCondensedJsonStringWriter > Writer = FCondensedJsonStringWriterFactory::Create(&OutputString);
-		check(FJsonSerializer::Serialize(Array, Writer));
-		check(InputString == OutputString);
+		REQUIRE(FJsonSerializer::Serialize(Array, Writer));
+		REQUIRE(InputString == OutputString);
 	}
 
 	// String Array Case
@@ -224,29 +218,29 @@ bool FJsonAutomationTest::RunTest(const FString& Parameters)
 
 		TArray< TSharedPtr<FJsonValue> > Array;
 		bool bSuccessful = FJsonSerializer::Deserialize(Reader, Array);
-		check(bSuccessful);
-		check(Array.Num() == 4);
-		check(Array[0].IsValid());
-		check(Array[1].IsValid());
-		check(Array[2].IsValid());
-		check(Array[3].IsValid());
+		REQUIRE(bSuccessful);
+		REQUIRE(Array.Num() == 4);
+		REQUIRE(Array[0].IsValid());
+		REQUIRE(Array[1].IsValid());
+		REQUIRE(Array[2].IsValid());
+		REQUIRE(Array[3].IsValid());
 
 		FString Text = Array[0]->AsString();
-		check(Text == TEXT("Some String1"));
+		REQUIRE(Text == TEXT("Some String1"));
 
 		Text = Array[1]->AsString();
-		check(Text == TEXT("Some String2"));
+		REQUIRE(Text == TEXT("Some String2"));
 
 		Text = Array[2]->AsString();
-		check(Text == TEXT("Some String3"));
+		REQUIRE(Text == TEXT("Some String3"));
 
 		Text = Array[3]->AsString();
-		check(Text == TEXT("Some String4"));
+		REQUIRE(Text == TEXT("Some String4"));
 
 		FString OutputString;
 		TSharedRef< FCondensedJsonStringWriter > Writer = FCondensedJsonStringWriterFactory::Create(&OutputString);
-		check(FJsonSerializer::Serialize(Array, Writer));
-		check(InputString == OutputString);
+		REQUIRE(FJsonSerializer::Serialize(Array, Writer));
+		REQUIRE(InputString == OutputString);
 	}
 
 	// Complex Array Case
@@ -273,46 +267,46 @@ bool FJsonAutomationTest::RunTest(const FString& Parameters)
 
 		TArray< TSharedPtr<FJsonValue> > Array;
 		bool bSuccessful = FJsonSerializer::Deserialize(Reader, Array);
-		check(bSuccessful);
-		check(Array.Num() == 6);
-		check(Array[0].IsValid());
-		check(Array[1].IsValid());
-		check(Array[2].IsValid());
-		check(Array[3].IsValid());
-		check(Array[4].IsValid());
-		check(Array[5].IsValid());
+		REQUIRE(bSuccessful);
+		REQUIRE(Array.Num() == 6);
+		REQUIRE(Array[0].IsValid());
+		REQUIRE(Array[1].IsValid());
+		REQUIRE(Array[2].IsValid());
+		REQUIRE(Array[3].IsValid());
+		REQUIRE(Array[4].IsValid());
+		REQUIRE(Array[5].IsValid());
 
 		FString Text = Array[0]->AsString();
-		check(Text == TEXT("Some String1"));
+		REQUIRE(Text == TEXT("Some String1"));
 
 		double Number = Array[1]->AsNumber();
-		check(Number == 10);
+		REQUIRE(Number == 10);
 
 		TSharedPtr< FJsonObject > Object = Array[2]->AsObject();
-		check(Object.IsValid());
-		check(Object->GetStringField(TEXT("Value")) == TEXT("Some String3"));
-		check(Object->GetStringField(TEXT("")) == TEXT("Empty Key"));
+		REQUIRE(Object.IsValid());
+		REQUIRE(Object->GetStringField(TEXT("Value")) == TEXT("Some String3"));
+		REQUIRE(Object->GetStringField(TEXT("")) == TEXT("Empty Key"));
 
 		const TArray<TSharedPtr< FJsonValue >>& InnerArray = Array[3]->AsArray();
-		check(InnerArray.Num() == 2);
-		check(Array[0].IsValid());
-		check(Array[1].IsValid());
+		REQUIRE(InnerArray.Num() == 2);
+		REQUIRE(Array[0].IsValid());
+		REQUIRE(Array[1].IsValid());
 
 		Text = InnerArray[0]->AsString();
-		check(Text == TEXT("Some String4"));
+		REQUIRE(Text == TEXT("Some String4"));
 
 		Text = InnerArray[1]->AsString();
-		check(Text == TEXT("Some String5"));
+		REQUIRE(Text == TEXT("Some String5"));
 
 		bool Boolean = Array[4]->AsBool();
-		check(Boolean == true);
+		REQUIRE(Boolean == true);
 
-		check(Array[5]->IsNull() == true);
+		REQUIRE(Array[5]->IsNull() == true);
 
 		FString OutputString;
 		TSharedRef< FCondensedJsonStringWriter > Writer = FCondensedJsonStringWriterFactory::Create(&OutputString);
-		check(FJsonSerializer::Serialize(Array, Writer));
-		check(InputString == OutputString);
+		REQUIRE(FJsonSerializer::Serialize(Array, Writer));
+		REQUIRE(InputString == OutputString);
 	}
 
 	// String Test
@@ -327,17 +321,18 @@ bool FJsonAutomationTest::RunTest(const FString& Parameters)
 
 		TSharedPtr<FJsonObject> Object;
 		bool bSuccessful = FJsonSerializer::Deserialize(Reader, Object);
-		check(bSuccessful);
-		check( Object.IsValid() );
+		REQUIRE(bSuccessful);
+		REQUIRE(Object.IsValid());
 
 		const TSharedPtr<FJsonValue>* Value = Object->Values.Find(TEXT("Value"));
-		check(Value && (*Value)->Type == EJson::String);
+		REQUIRE(Value);
+		REQUIRE((*Value)->Type == EJson::String);
 		const FString String = (*Value)->AsString();
-		check(String == TEXT("Some String, Escape Chars: \\, \", /, \b, \f, \n, \r, \t, +"));
+		REQUIRE(String == TEXT("Some String, Escape Chars: \\, \", /, \b, \f, \n, \r, \t, +"));
 
 		FString OutputString;
 		TSharedRef< FCondensedJsonStringWriter > Writer = FCondensedJsonStringWriterFactory::Create( &OutputString );
-		verify( FJsonSerializer::Serialize( Object.ToSharedRef(), Writer ) );
+		REQUIRE(FJsonSerializer::Serialize( Object.ToSharedRef(), Writer ));
 
 		const FString TestOutput =
 			TEXT(
@@ -345,13 +340,13 @@ bool FJsonAutomationTest::RunTest(const FString& Parameters)
 					"\"Value\":\"Some String, Escape Chars: \\\\, \\\", /, \\b, \\f, \\n, \\r, \\t, +\""
 				"}"
 			);
-		check(OutputString == TestOutput);
+		REQUIRE(OutputString == TestOutput);
 	}
 
 	// String Test UTF8
 	{
 		// UTF8TEXT will prepend the first u8 literal specifier
-		// UTF8TEXT does a cast so we can't add it each line and still get the literals to concatenate
+		// UTF8TEXT does a cast, so we can't add it each line and still get the literals to concatenate
 		const UTF8CHAR* InputString = UTF8TEXT(
 			"{"
 				u8"\"Value\":\"Some String, Escape Chars: \\\\, \\\", \\/, \\b, \\f, \\n, \\r, \\t, \\u002B\\uD83D\\uDE10\","
@@ -365,39 +360,39 @@ bool FJsonAutomationTest::RunTest(const FString& Parameters)
 
 		TSharedPtr<FJsonObject> Object;
 		bool bSuccessful = FJsonSerializer::Deserialize(Reader, Object);
-		check(bSuccessful);
-		check( Object.IsValid() );
+		REQUIRE(bSuccessful);
+		REQUIRE(Object.IsValid());
 
 		{
 			const TSharedPtr<FJsonValue>* Value = Object->Values.Find(TEXT("Value"));
-			check(Value && (*Value)->Type == EJson::String);
+			REQUIRE((Value && (*Value)->Type == EJson::String));
 			const FString String = (*Value)->AsString();
-			check(String == TEXT("Some String, Escape Chars: \\, \", /, \b, \f, \n, \r, \t, +😐"));
+			REQUIRE(String == TEXT("Some String, Escape Chars: \\, \", /, \b, \f, \n, \r, \t, +😐"));
 		}
 		{
 			const TSharedPtr<FJsonValue>* Value = Object->Values.Find(TEXT("Value1"));
-			check(Value && (*Value)->Type == EJson::String);
+			REQUIRE((Value && (*Value)->Type == EJson::String));
 			const FString String = (*Value)->AsString();
-			check(String == TEXT("Greek String, Σὲ γνωρίζω ἀπὸ τὴν κόψη"));
+			REQUIRE(String == TEXT("Greek String, Σὲ γνωρίζω ἀπὸ τὴν κόψη"));
 		}
 		{
 			const TSharedPtr<FJsonValue>* Value = Object->Values.Find(TEXT("Value2"));
-			check(Value && (*Value)->Type == EJson::String);
+			REQUIRE((Value && (*Value)->Type == EJson::String));
 			const FString String = (*Value)->AsString();
-			check(String == TEXT("Thai String, สิบสองกษัตริย์ก่อนหน้าแลถัดไป"));
+			REQUIRE(String == TEXT("Thai String, สิบสองกษัตริย์ก่อนหน้าแลถัดไป"));
 		}
 		{
 			const TSharedPtr<FJsonValue>* Value = Object->Values.Find(TEXT("Value3"));
-			check(Value && (*Value)->Type == EJson::String);
+			REQUIRE((Value && (*Value)->Type == EJson::String));
 			const FString String = (*Value)->AsString();
-			check(String == TEXT("Hello world, Καλημέρα κόσμε, コンニチハ"));
+			REQUIRE(String == TEXT("Hello world, Καλημέρα κόσμε, コンニチハ"));
 		}
 
 		FString OutputString;
 		TSharedRef< FCondensedJsonStringWriter > Writer = FCondensedJsonStringWriterFactory::Create( &OutputString );
-		verify( FJsonSerializer::Serialize( Object.ToSharedRef(), Writer ) );
+		REQUIRE(FJsonSerializer::Serialize( Object.ToSharedRef(), Writer ));
 
-		// Note: The literal prefix for the string (u8, L), must be present for every contatenated string, not just the first one
+		// Note: The literal prefix for the string (u8, L), must be present for every concatenated string, not just the first one
 		const FString TestOutput =
 			TEXT("{")
 				TEXT("\"Value\":\"Some String, Escape Chars: \\\\, \\\", /, \\b, \\f, \\n, \\r, \\t, +😐\",")
@@ -405,7 +400,7 @@ bool FJsonAutomationTest::RunTest(const FString& Parameters)
 				TEXT("\"Value2\":\"Thai String, สิบสองกษัตริย์ก่อนหน้าแลถัดไป\",")
 				TEXT("\"Value3\":\"Hello world, Καλημέρα κόσμε, コンニチハ\"")
 			TEXT("}");
-		check(OutputString == TestOutput);
+		REQUIRE(OutputString == TestOutput);
 	}
 
 	// Number Test
@@ -424,21 +419,21 @@ bool FJsonAutomationTest::RunTest(const FString& Parameters)
 
 		TSharedPtr<FJsonObject> Object;
 		bool bSuccessful = FJsonSerializer::Deserialize(Reader, Object);
-		check(bSuccessful);
-		check( Object.IsValid() );
+		REQUIRE(bSuccessful);
+		REQUIRE(Object.IsValid());
 
 		double TestValues[] = {2.544e+15, -0.544e-2, 251e3, -0.0, 843};
 		for (int32 i = 0; i < 5; ++i)
 		{
 			const TSharedPtr<FJsonValue>* Value = Object->Values.Find(FString::Printf(TEXT("Value%i"), i + 1));
-			check(Value && (*Value)->Type == EJson::Number);
+			REQUIRE((Value && (*Value)->Type == EJson::Number));
 			const double Number = (*Value)->AsNumber();
-			check(Number == TestValues[i]);
+			REQUIRE(Number == TestValues[i]);
 		}
 
 		FString OutputString;
 		TSharedRef< FCondensedJsonStringWriter > Writer = FCondensedJsonStringWriterFactory::Create( &OutputString );
-		verify( FJsonSerializer::Serialize( Object.ToSharedRef(), Writer ) );
+		REQUIRE(FJsonSerializer::Serialize( Object.ToSharedRef(), Writer ));
 
 		// %g isn't standardized, so we use the same %g format that is used inside PrintJson instead of hardcoding the values here
 		const FString TestOutput = FString::Printf(
@@ -452,7 +447,7 @@ bool FJsonAutomationTest::RunTest(const FString& Parameters)
 				"}"
 			),
 			TestValues[0], TestValues[1], TestValues[2], TestValues[3], TestValues[4]);
-		check(OutputString == TestOutput);
+		REQUIRE(OutputString == TestOutput);
 	}
 
 	// Boolean/Null Test
@@ -471,30 +466,30 @@ bool FJsonAutomationTest::RunTest(const FString& Parameters)
 
 		TSharedPtr<FJsonObject> Object;
 		bool bSuccessful = FJsonSerializer::Deserialize(Reader, Object);
-		check(bSuccessful);
-		check( Object.IsValid() );
+		REQUIRE(bSuccessful);
+		REQUIRE(Object.IsValid());
 
 		bool TestValues[] = {true, true, false};
 		for (int32 i = 0; i < 5; ++i)
 		{
 			const TSharedPtr<FJsonValue>* Value = Object->Values.Find(FString::Printf(TEXT("Value%i"), i + 1));
-			check(Value);
+			REQUIRE(Value);
 			if (i < 3)
 			{
-				check((*Value)->Type == EJson::Boolean);
+				REQUIRE((*Value)->Type == EJson::Boolean);
 				const bool Bool = (*Value)->AsBool();
-				check(Bool == TestValues[i]);
+				REQUIRE(Bool == TestValues[i]);
 			}
 			else
 			{
-				check((*Value)->Type == EJson::Null);
-				check((*Value)->IsNull());
+				REQUIRE((*Value)->Type == EJson::Null);
+				REQUIRE((*Value)->IsNull());
 			}
 		}
 
 		FString OutputString;
 		TSharedRef< FCondensedJsonStringWriter > Writer = FCondensedJsonStringWriterFactory::Create( &OutputString );
-		verify( FJsonSerializer::Serialize( Object.ToSharedRef(), Writer ) );
+		REQUIRE(FJsonSerializer::Serialize( Object.ToSharedRef(), Writer ));
 
 		const FString TestOutput =
 			TEXT(
@@ -506,7 +501,7 @@ bool FJsonAutomationTest::RunTest(const FString& Parameters)
 					"\"Value5\":null"
 				"}"
 			);
-		check(OutputString == TestOutput);
+		REQUIRE(OutputString == TestOutput);
 	}
 
 	// Object Test && extra whitespace test
@@ -539,42 +534,42 @@ bool FJsonAutomationTest::RunTest(const FString& Parameters)
 
 		TSharedPtr<FJsonObject> Object;
 		bool bSuccessful = FJsonSerializer::Deserialize(Reader, Object);
-		check(bSuccessful);
-		check( Object.IsValid() );
+		REQUIRE(bSuccessful);
+		REQUIRE(Object.IsValid());
 
 		const TSharedPtr<FJsonValue>* InnerValueFail = Object->Values.Find(TEXT("InnerValue"));
-		check(!InnerValueFail);
+		REQUIRE(!InnerValueFail);
 
 		const TSharedPtr<FJsonValue>* ObjectValue = Object->Values.Find(TEXT("Object"));
-		check(ObjectValue && (*ObjectValue)->Type == EJson::Object);
+		REQUIRE((ObjectValue && (*ObjectValue)->Type == EJson::Object));
 		const TSharedPtr<FJsonObject> InnerObject = (*ObjectValue)->AsObject();
-		check(InnerObject.IsValid());
+		REQUIRE(InnerObject.IsValid());
 
 		{
 			const TSharedPtr<FJsonValue>* NestedValueValue = InnerObject->Values.Find(TEXT("NestedValue"));
-			check(NestedValueValue && (*NestedValueValue)->Type == EJson::Null);
-			check((*NestedValueValue)->IsNull());
+			REQUIRE((NestedValueValue && (*NestedValueValue)->Type == EJson::Null));
+			REQUIRE((*NestedValueValue)->IsNull());
 
 			const TSharedPtr<FJsonValue>* NestedObjectValue = InnerObject->Values.Find(TEXT("NestedObject"));
-			check(NestedObjectValue && (*NestedObjectValue)->Type == EJson::Object);
+			REQUIRE((NestedObjectValue && (*NestedObjectValue)->Type == EJson::Object));
 			const TSharedPtr<FJsonObject> InnerInnerObject = (*NestedObjectValue)->AsObject();
-			check(InnerInnerObject.IsValid());
+			REQUIRE(InnerInnerObject.IsValid());
 
 			{
 				const TSharedPtr<FJsonValue>* NestedValueValueFail = InnerInnerObject->Values.Find(TEXT("NestedValue"));
-				check(!NestedValueValueFail);
+				REQUIRE(!NestedValueValueFail);
 			}
 		}
 
 		const TSharedPtr<FJsonValue>* ValueValue = Object->Values.Find(TEXT("Value"));
-		check(ValueValue && (*ValueValue)->Type == EJson::Boolean);
+		REQUIRE((ValueValue && (*ValueValue)->Type == EJson::Boolean));
 		const bool Bool = (*ValueValue)->AsBool();
-		check(Bool);
+		REQUIRE(Bool);
 
 		FString OutputString;
 		TSharedRef< FCondensedJsonStringWriter > Writer = FCondensedJsonStringWriterFactory::Create( &OutputString );
-		verify( FJsonSerializer::Serialize( Object.ToSharedRef(), Writer ) );
-		check(OutputString == InputString);
+		REQUIRE(FJsonSerializer::Serialize( Object.ToSharedRef(), Writer ));
+		REQUIRE(OutputString == InputString);
 	}
 
 	// Array Test
@@ -599,41 +594,41 @@ bool FJsonAutomationTest::RunTest(const FString& Parameters)
 
 		TSharedPtr<FJsonObject> Object;
 		bool bSuccessful = FJsonSerializer::Deserialize(Reader, Object);
-		check(bSuccessful);
-		check( Object.IsValid() );
+		REQUIRE(bSuccessful);
+		REQUIRE(Object.IsValid());
 
 		const TSharedPtr<FJsonValue>* InnerValueFail = Object->Values.Find(TEXT("InnerValue"));
-		check(!InnerValueFail);
+		REQUIRE(!InnerValueFail);
 
 		const TSharedPtr<FJsonValue>* ArrayValue = Object->Values.Find(TEXT("Array"));
-		check(ArrayValue && (*ArrayValue)->Type == EJson::Array);
+		REQUIRE((ArrayValue && (*ArrayValue)->Type == EJson::Array));
 		const TArray< TSharedPtr<FJsonValue> > Array = (*ArrayValue)->AsArray();
-		check(Array.Num() == 8);
+		REQUIRE(Array.Num() == 8);
 
 		EJson ValueTypes[] = {EJson::Array, EJson::String, EJson::String, EJson::Null,
 			EJson::Boolean, EJson::Boolean, EJson::Number, EJson::Object};
 		for (int32 i = 0; i < Array.Num(); ++i)
 		{
 			const TSharedPtr<FJsonValue>& Value = Array[i];
-			check(Value.IsValid());
-			check(Value->Type == ValueTypes[i]);
+			REQUIRE(Value.IsValid());
+			REQUIRE(Value->Type == ValueTypes[i]);
 		}
 
 		const TArray< TSharedPtr<FJsonValue> >& InnerArray = Array[0]->AsArray();
-		check(InnerArray.Num() == 0);
-		check(Array[1]->AsString() == TEXT("Some String"));
-		check(Array[2]->AsString() == TEXT("Another String"));
-		check(Array[3]->IsNull());
-		check(Array[4]->AsBool());
-		check(!Array[5]->AsBool());
-		check(FMath::Abs(Array[6]->AsNumber() - 45.f) < KINDA_SMALL_NUMBER);
+		REQUIRE(InnerArray.Num() == 0);
+		REQUIRE(Array[1]->AsString() == TEXT("Some String"));
+		REQUIRE(Array[2]->AsString() == TEXT("Another String"));
+		REQUIRE(Array[3]->IsNull());
+		REQUIRE(Array[4]->AsBool());
+		REQUIRE(!Array[5]->AsBool());
+		REQUIRE(FMath::Abs(Array[6]->AsNumber() - 45.f) < KINDA_SMALL_NUMBER);
 		const TSharedPtr<FJsonObject> InnerObject = Array[7]->AsObject();
-		check(InnerObject.IsValid());
+		REQUIRE(InnerObject.IsValid());
 
 		FString OutputString;
 		TSharedRef< FCondensedJsonStringWriter > Writer = FCondensedJsonStringWriterFactory::Create( &OutputString );
-		verify( FJsonSerializer::Serialize( Object.ToSharedRef(), Writer ) );
-		check(OutputString == InputString);
+		REQUIRE(FJsonSerializer::Serialize( Object.ToSharedRef(), Writer ));
+		REQUIRE(OutputString == InputString);
 	}
 
 	// Pretty Print Test
@@ -662,15 +657,15 @@ bool FJsonAutomationTest::RunTest(const FString& Parameters)
 		TSharedRef< TJsonReader<> > Reader = TJsonReaderFactory<>::Create( InputString );
 
 		TSharedPtr<FJsonObject> Object;
-		verify( FJsonSerializer::Deserialize( Reader, Object ) );
-		check( Object.IsValid() );
+		REQUIRE(FJsonSerializer::Deserialize( Reader, Object ));
+		REQUIRE(Object.IsValid());
 
 		FString OutputString;
 		TSharedRef< FPrettyJsonStringWriter > Writer = FPrettyJsonStringWriterFactory::Create( &OutputString );
-		verify( FJsonSerializer::Serialize( Object.ToSharedRef(), Writer ) );
-		check(OutputString == InputString);
+		REQUIRE(FJsonSerializer::Serialize( Object.ToSharedRef(), Writer ));
+		REQUIRE(OutputString == InputString);
 	}
-	  
+
 	// Line and Character # test
 	{
 		const FString InputString =
@@ -690,20 +685,20 @@ bool FJsonAutomationTest::RunTest(const FString& Parameters)
 		TSharedRef< TJsonReader<> > Reader = TJsonReaderFactory<>::Create( InputString );
 
 		EJsonNotation Notation = EJsonNotation::Null;
-		verify( Reader->ReadNext( Notation ) && Notation == EJsonNotation::ObjectStart );
-		check( Reader->GetLineNumber() == 1 && Reader->GetCharacterNumber() == 1 );
+		REQUIRE(( Reader->ReadNext( Notation ) && Notation == EJsonNotation::ObjectStart ));
+		REQUIRE(( Reader->GetLineNumber() == 1 && Reader->GetCharacterNumber() == 1 ));
 
-		verify( Reader->ReadNext( Notation ) && Notation == EJsonNotation::String );
-		check( Reader->GetLineNumber() == 2 && Reader->GetCharacterNumber() == 17 );
+		REQUIRE(( Reader->ReadNext( Notation ) && Notation == EJsonNotation::String ));
+		REQUIRE(( Reader->GetLineNumber() == 2 && Reader->GetCharacterNumber() == 17 ));
 
-		verify( Reader->ReadNext( Notation ) && Notation == EJsonNotation::ArrayStart );
-		check( Reader->GetLineNumber() == 4 && Reader->GetCharacterNumber() == 2 );
+		REQUIRE(( Reader->ReadNext( Notation ) && Notation == EJsonNotation::ArrayStart ));
+		REQUIRE(( Reader->GetLineNumber() == 4 && Reader->GetCharacterNumber() == 2 ));
 
-		verify( Reader->ReadNext( Notation ) && Notation == EJsonNotation::Number );
-		check( Reader->GetLineNumber() == 5 && Reader->GetCharacterNumber() == 7 );
+		REQUIRE(( Reader->ReadNext( Notation ) && Notation == EJsonNotation::Number ));
+		REQUIRE(( Reader->GetLineNumber() == 5 && Reader->GetCharacterNumber() == 7 ));
 
-		verify( Reader->ReadNext( Notation ) && Notation == EJsonNotation::Boolean );
-		check( Reader->GetLineNumber() == 6 && Reader->GetCharacterNumber() == 6 );
+		REQUIRE(( Reader->ReadNext( Notation ) && Notation == EJsonNotation::Boolean ));
+		REQUIRE(( Reader->GetLineNumber() == 6 && Reader->GetCharacterNumber() == 6 ));
 	}
 
 	// Failure Cases
@@ -827,8 +822,8 @@ bool FJsonAutomationTest::RunTest(const FString& Parameters)
 		TSharedRef< TJsonReader<> > Reader = TJsonReaderFactory<>::Create( FailureInputs[i] );
 
 		TSharedPtr<FJsonObject> Object;
-		verify( FJsonSerializer::Deserialize( Reader, Object ) == false );
-		check( !Object.IsValid() );
+		REQUIRE(FJsonSerializer::Deserialize( Reader, Object ) == false);
+		REQUIRE(!Object.IsValid());
 	}
 
 	// TryGetNumber tests
@@ -1004,8 +999,6 @@ bool FJsonAutomationTest::RunTest(const FString& Parameters)
 			TestEqual(TEXT("TryGetNumber-UInt32 Half rounds to next integer"), IntVal, 1U);
 		}
 	}
-
-	return true;
 }
 
-#endif //WITH_DEV_AUTOMATION_TESTS
+#endif // WITH_TESTS
