White space adjustments

diff --git a/Engine/Source/Runtime/Json/Public/Serialization/JsonSerializable.h b/Engine/Source/Runtime/Json/Public/Serialization/JsonSerializable.h
--- a/Engine/Source/Runtime/Json/Public/Serialization/JsonSerializable.h
+++ b/Engine/Source/Runtime/Json/Public/Serialization/JsonSerializable.h
@@ -77,157 +77,160 @@ struct FJsonSerializable
 
 namespace UE::JsonArray
 {
-	namespace Private
+
+namespace Private
+{
+
+template<typename T, typename CharType>
+inline bool FromJson(TArray<T>& OutArray, TStringView<CharType> JsonString)
+{
+	OutArray.Reset();
+
+	TArray<TSharedPtr<FJsonValue>> ArrayValues;
+	TSharedRef<TJsonReader<CharType>> JsonReader = TJsonReaderFactory<CharType>::CreateFromView(JsonString);
+	if (FJsonSerializer::Deserialize(JsonReader, ArrayValues))
 	{
-		template<typename T, typename CharType>
-		inline bool FromJson(TArray<T>& OutArray, TStringView<CharType> JsonString)
+		for (const TSharedPtr<FJsonValue>& Value : ArrayValues)
 		{
-			OutArray.Reset();
-
-			TArray<TSharedPtr<FJsonValue>> ArrayValues;
-			TSharedRef<TJsonReader<CharType>> JsonReader = TJsonReaderFactory<CharType>::CreateFromView(JsonString);
-			if (FJsonSerializer::Deserialize(JsonReader, ArrayValues))
+			TSharedPtr<FJsonObject>* ArrayEntry;
+			if (Value.IsValid() && Value->TryGetObject(ArrayEntry))
 			{
-				for (const TSharedPtr<FJsonValue>& Value : ArrayValues)
+				if (ArrayEntry && ArrayEntry->IsValid())
 				{
-					TSharedPtr<FJsonObject>* ArrayEntry;
-					if (Value.IsValid() && Value->TryGetObject(ArrayEntry))
-					{
-						if (ArrayEntry && ArrayEntry->IsValid())
-						{
-							FJsonSerializerReader Serializer(*ArrayEntry);
-							OutArray.Add_GetRef(T()).Serialize(Serializer, false);
-						}
-						else
-						{
-							UE_LOG(LogJson, Error, TEXT("Failed to parse Json from array"));
-							return false;
-						}
-					}
+					FJsonSerializerReader Serializer(*ArrayEntry);
+					OutArray.Add_GetRef(T()).Serialize(Serializer, false);
+				}
+				else
+				{
+					UE_LOG(LogJson, Error, TEXT("Failed to parse Json from array"));
+					return false;
 				}
-
-				return true;
 			}
-
-			return false;
 		}
 
-		using ReturnStringArgs = TTuple<FString* /*OutValue*/, bool /*bPrettyPrint*/>;
+		return true;
+	}
+
+	return false;
+}
 
-		using PrettyWriter = TSharedRef<TJsonWriter<>>;
-		using CondensedWriter = TSharedRef<TJsonWriter<TCHAR, TCondensedJsonPrintPolicy<TCHAR>>>;
-		using WriterVariants = TVariant<PrettyWriter, CondensedWriter>;
+using ReturnStringArgs = TTuple<FString* /*OutValue*/, bool /*bPrettyPrint*/>;
 
-		using ToJsonVariantArgs = TVariant<ReturnStringArgs, WriterVariants>;
+using PrettyWriter = TSharedRef<TJsonWriter<>>;
+using CondensedWriter = TSharedRef<TJsonWriter<TCHAR, TCondensedJsonPrintPolicy<TCHAR>>>;
+using WriterVariants = TVariant<PrettyWriter, CondensedWriter>;
 
-		using PrettySerializer = FJsonSerializerWriter<>;
-		using CondensedSerializer = FJsonSerializerWriter<TCHAR, TCondensedJsonPrintPolicy<TCHAR>>;
+using ToJsonVariantArgs = TVariant<ReturnStringArgs, WriterVariants>;
 
-		template<typename T>
-		inline void ToJson(TArray<T>& InArray, const ToJsonVariantArgs& InArgs)
-		{
-			using PrettySerializerAndWriter = TTuple<PrettySerializer, PrettyWriter>;
-			using CondensedSerializerAndWriter = TTuple<CondensedSerializer, CondensedWriter>;
+using PrettySerializer = FJsonSerializerWriter<>;
+using CondensedSerializer = FJsonSerializerWriter<TCHAR, TCondensedJsonPrintPolicy<TCHAR>>;
+
+template<typename T>
+inline void ToJson(TArray<T>& InArray, const ToJsonVariantArgs& InArgs)
+{
+	using PrettySerializerAndWriter = TTuple<PrettySerializer, PrettyWriter>;
+	using CondensedSerializerAndWriter = TTuple<CondensedSerializer, CondensedWriter>;
 
-			using SerializerVariant = TVariant<PrettySerializerAndWriter, CondensedSerializerAndWriter>;
+	using SerializerVariant = TVariant<PrettySerializerAndWriter, CondensedSerializerAndWriter>;
 
-			SerializerVariant SerializerToUse = ::Visit([](auto& StoredValue)
+	SerializerVariant SerializerToUse = ::Visit([](auto& StoredValue)
+		{
+			using StoredValueType = std::decay_t<decltype(StoredValue)>;
+			if constexpr (std::is_same_v<StoredValueType, ReturnStringArgs>)
+			{
+				if (StoredValue.template Get<1>())
+				{
+					PrettyWriter NewWriter = TJsonWriterFactory<>::Create(StoredValue.template Get<0>());
+					return SerializerVariant(TInPlaceType<PrettySerializerAndWriter>(), PrettySerializer(NewWriter), NewWriter);
+				}
+				else
 				{
-					using StoredValueType = std::decay_t<decltype(StoredValue)>;
-					if constexpr (std::is_same_v<StoredValueType, ReturnStringArgs>)
+					CondensedWriter NewWriter = TJsonWriterFactory<TCHAR, TCondensedJsonPrintPolicy<TCHAR>>::Create(StoredValue.template Get<0>());
+					return SerializerVariant(TInPlaceType<CondensedSerializerAndWriter>(), CondensedSerializer(NewWriter), NewWriter);
+				}
+			}
+			else
+			{
+				return ::Visit([](auto& StoredWriter)
 					{
-						if (StoredValue.template Get<1>())
+						using StoredWriterType = std::decay_t<decltype(StoredWriter)>;
+						if constexpr (std::is_same_v<StoredWriterType, PrettyWriter>)
 						{
-							PrettyWriter NewWriter = TJsonWriterFactory<>::Create(StoredValue.template Get<0>());
-							return SerializerVariant(TInPlaceType<PrettySerializerAndWriter>(), PrettySerializer(NewWriter), NewWriter);
+							return SerializerVariant(TInPlaceType<PrettySerializerAndWriter>(), PrettySerializer(StoredWriter), StoredWriter);
 						}
-						else
+						else if constexpr (std::is_same_v<StoredWriterType, CondensedWriter>)
 						{
-							CondensedWriter NewWriter = TJsonWriterFactory<TCHAR, TCondensedJsonPrintPolicy<TCHAR>>::Create(StoredValue.template Get<0>());
-							return SerializerVariant(TInPlaceType<CondensedSerializerAndWriter>(), CondensedSerializer(NewWriter), NewWriter);
+							return SerializerVariant(TInPlaceType<CondensedSerializerAndWriter>(), CondensedSerializer(StoredWriter), StoredWriter);
 						}
-					}
-					else
-					{
-						return ::Visit([](auto& StoredWriter)
-							{
-								using StoredWriterType = std::decay_t<decltype(StoredWriter)>;
-								if constexpr (std::is_same_v<StoredWriterType, PrettyWriter>)
-								{
-									return SerializerVariant(TInPlaceType<PrettySerializerAndWriter>(), PrettySerializer(StoredWriter), StoredWriter);
-								}
-								else if constexpr (std::is_same_v<StoredWriterType, CondensedWriter>)
-								{
-									return SerializerVariant(TInPlaceType<CondensedSerializerAndWriter>(), CondensedSerializer(StoredWriter), StoredWriter);
-								}
-							}, StoredValue);
-					}
-				}, InArgs);
-
-
-			const bool bCloseWriter = InArgs.IsType<ReturnStringArgs>();
-
-			::Visit([bCloseWriter, &InArray](auto& StoredSerializer)
-				{
-					StoredSerializer.template Get<0>().StartArray();
+					}, StoredValue);
+			}
+		}, InArgs);
 
-					for (T& ArrayEntry : InArray)
-					{
-						ArrayEntry.Serialize(StoredSerializer.template Get<0>(), false);
-					}
 
-					StoredSerializer.template Get<0>().EndArray();
+	const bool bCloseWriter = InArgs.IsType<ReturnStringArgs>();
 
-					if (bCloseWriter)
-					{
-						StoredSerializer.template Get<1>()->Close();
-					}
-				}, SerializerToUse);
-		}
-	}
+	::Visit([bCloseWriter, &InArray](auto& StoredSerializer)
+		{
+			StoredSerializer.template Get<0>().StartArray();
 
-	template<typename T>
-	static bool FromJson(TArray<T>& OutArray, const FString& JsonString)
-	{
-		return Private::FromJson(OutArray, FStringView(JsonString));
-	}
+			for (T& ArrayEntry : InArray)
+			{
+				ArrayEntry.Serialize(StoredSerializer.template Get<0>(), false);
+			}
 
-	template<typename T>
-	static bool FromJson(TArray<T>& OutArray, FString&& JsonString)
-	{
-		return Private::FromJson(OutArray, FStringView(MoveTemp(JsonString)));
-	}
+			StoredSerializer.template Get<0>().EndArray();
 
-	template<typename T>
-	static bool FromJson(TArray<T>& OutArray, FUtf8StringView JsonStringView)
-	{
-		return Private::FromJson(OutArray, JsonStringView);
-	}
+			if (bCloseWriter)
+			{
+				StoredSerializer.template Get<1>()->Close();
+			}
+		}, SerializerToUse);
+}
+} // namespace Private
 
-	template<typename T>
-	static bool FromJson(TArray<T>& OutArray, FWideStringView JsonStringView)
-	{
-		return Private::FromJson(OutArray, JsonStringView);
-	}
+template<typename T>
+static bool FromJson(TArray<T>& OutArray, const FString& JsonString)
+{
+	return Private::FromJson(OutArray, FStringView(JsonString));
+}
 
-	/* non-const due to T::Serialize being a non-const function */
-	template<typename T>
-	static const FString ToJson(TArray<T>& InArray, const bool bPrettyPrint = true)
-	{
-		FString JsonStr;
-		Private::ToJson(InArray, Private::ToJsonVariantArgs(TInPlaceType<Private::ReturnStringArgs>(), &JsonStr, bPrettyPrint));
-		return JsonStr;
-	}
+template<typename T>
+static bool FromJson(TArray<T>& OutArray, FString&& JsonString)
+{
+	return Private::FromJson(OutArray, FStringView(MoveTemp(JsonString)));
+}
 
-	template<typename T>
-	static void ToJson(TArray<T>& InArray, Private::PrettyWriter& JsonWriter)
-	{
-		Private::ToJson(InArray, Private::ToJsonVariantArgs(TInPlaceType<Private::WriterVariants>(), Private::WriterVariants(TInPlaceType<Private::PrettyWriter>(), JsonWriter)));
-	}
+template<typename T>
+static bool FromJson(TArray<T>& OutArray, FUtf8StringView JsonStringView)
+{
+	return Private::FromJson(OutArray, JsonStringView);
+}
 
-	template<typename T>
-	static void ToJson(TArray<T>& InArray, Private::CondensedWriter& JsonWriter)
-	{
-		Private::ToJson(InArray, Private::ToJsonVariantArgs(TInPlaceType<Private::WriterVariants>(), Private::WriterVariants(TInPlaceType<Private::CondensedWriter>(), JsonWriter)));
-	}
-}
\ No newline at end of file
+template<typename T>
+static bool FromJson(TArray<T>& OutArray, FWideStringView JsonStringView)
+{
+	return Private::FromJson(OutArray, JsonStringView);
+}
+
+/* non-const due to T::Serialize being a non-const function */
+template<typename T>
+static const FString ToJson(TArray<T>& InArray, const bool bPrettyPrint = true)
+{
+	FString JsonStr;
+	Private::ToJson(InArray, Private::ToJsonVariantArgs(TInPlaceType<Private::ReturnStringArgs>(), &JsonStr, bPrettyPrint));
+	return JsonStr;
+}
+
+template<typename T>
+static void ToJson(TArray<T>& InArray, Private::PrettyWriter& JsonWriter)
+{
+	Private::ToJson(InArray, Private::ToJsonVariantArgs(TInPlaceType<Private::WriterVariants>(), Private::WriterVariants(TInPlaceType<Private::PrettyWriter>(), JsonWriter)));
+}
+
+template<typename T>
+static void ToJson(TArray<T>& InArray, Private::CondensedWriter& JsonWriter)
+{
+	Private::ToJson(InArray, Private::ToJsonVariantArgs(TInPlaceType<Private::WriterVariants>(), Private::WriterVariants(TInPlaceType<Private::CondensedWriter>(), JsonWriter)));
+}
+
+} // namespace UE::JsonArray
\ No newline at end of file
