Implemented a `JsonValueVariant = TVariant<bool, TVariant<int32, uint32, int64, float, double>, FString`> for the 3 possible json values that aren't object and array.

This allows TMap<FString, JsonValueVariant> to contain variant types of values instead of being forced to be `FString` values

Also implemented Optional Map and Optional Array macros

diff --git a/Engine/Source/Runtime/Json/Private/Dom/JsonValue.cpp b/Engine/Source/Runtime/Json/Private/Dom/JsonValue.cpp
--- a/Engine/Source/Runtime/Json/Private/Dom/JsonValue.cpp
+++ b/Engine/Source/Runtime/Json/Private/Dom/JsonValue.cpp
@@ -175,7 +175,16 @@ bool FJsonValue::CompareEqual( const FJsonValue& Lhs, const FJsonValue& Rhs )
 {
 	if (Lhs.Type != Rhs.Type)
 	{
-		return false;
+		const bool bLhsIsSimpleVariant = Lhs.Type == EJson::Boolean || Lhs.Type == EJson::Number || Lhs.Type == EJson::String;
+		const bool bRhsIsSimpleVariant = Rhs.Type == EJson::Boolean || Rhs.Type == EJson::Number || Rhs.Type == EJson::String;
+		if (bLhsIsSimpleVariant && bRhsIsSimpleVariant)
+		{
+			return UE::Json::ToSimpleJsonVariant(Lhs) == UE::Json::ToSimpleJsonVariant(Rhs);
+		}
+		else
+		{
+			return false;
+		}
 	}
 
 	switch (Lhs.Type)
diff --git a/Engine/Source/Runtime/Json/Private/Serialization/JsonSerializerReader.cpp b/Engine/Source/Runtime/Json/Private/Serialization/JsonSerializerReader.cpp
--- a/Engine/Source/Runtime/Json/Private/Serialization/JsonSerializerReader.cpp
+++ b/Engine/Source/Runtime/Json/Private/Serialization/JsonSerializerReader.cpp
@@ -2,6 +2,8 @@
 
 #include "Serialization/JsonSerializerReader.h"
 
+#include "Dom/JsonValue.h"
+
 FJsonSerializerReader::FJsonSerializerReader(TSharedPtr<FJsonObject> InJsonObject)
 	: JsonObject(InJsonObject)
 {
@@ -128,6 +130,15 @@ void FJsonSerializerReader::Serialize(FStringView Name, FDateTime& Value)
 	}
 }
 
+void FJsonSerializerReader::Serialize(FStringView Name, JsonSimpleValueVariant& InVariant)
+{
+	TSharedPtr<FJsonValue> FieldValue = JsonObject->TryGetField(Name);
+	if (FieldValue.IsValid())
+	{
+		InVariant = UE::Json::ToSimpleJsonVariant(*FieldValue);
+	}
+}
+
 void FJsonSerializerReader::SerializeArray(FJsonSerializableArray& Array) 
 {
 	// @todo - higher level serialization is expecting a Json Object
@@ -250,6 +261,22 @@ void FJsonSerializerReader::SerializeMap(FStringView Name, FJsonSerializableKeyV
 	}
 }
 
+void FJsonSerializerReader::SerializeMap(FStringView Name, FJsonSerializableKeySimpleValueVariantMap& Map)
+{
+	if (JsonObject->HasTypedField<EJson::Object>(Name))
+	{
+		TSharedPtr<FJsonObject> JsonMap = JsonObject->GetObjectField(Name);
+		// Iterate all of the keys and their values
+		for (const TPair<FString, TSharedPtr<FJsonValue>>& Pair : JsonMap->Values)
+		{
+			if (Pair.Value.IsValid())
+			{
+				Map.Add(Pair.Key, UE::Json::ToSimpleJsonVariant(*Pair.Value));
+			}
+		}
+	}
+}
+
 void FJsonSerializerReader::SerializeSimpleMap(FJsonSerializableKeyValueMap& Map) 
 {
 	// Iterate all of the keys and their values, only taking simple types (not array/object), all in string form
diff --git a/Engine/Source/Runtime/Json/Private/Tests/JsonTests.cpp b/Engine/Source/Runtime/Json/Private/Tests/JsonTests.cpp
--- a/Engine/Source/Runtime/Json/Private/Tests/JsonTests.cpp
+++ b/Engine/Source/Runtime/Json/Private/Tests/JsonTests.cpp
@@ -161,6 +161,180 @@ TEST_CASE_NAMED(FJsonAutomationTest, "System::Engine::FileSystem::JSON", "[Appli
 		REQUIRE(InputString == OutputString);
 	}
 
+	// Variant Value Comparison Equality Test
+	{
+		/* comparing: "Type1_Type2_#" */
+		const FString StoredAsType1 =
+			TEXT(
+				"{"
+					"\"bool_string_0\" : false,"
+					"\"bool_string_1\" : true,"
+					"\"bool_string_2\" : false,"
+					"\"bool_string_3\" : true,"
+
+					"\"int_string_0\" : 10,"
+					"\"int_string_1\" : 100,"
+					
+					"\"float_string_0\" : 10.123,"
+					"\"float_string_1\" : 100.34,"
+
+					"\"string_string_0\" : \"foo1\","
+					"\"string_string_1\" : \"foo2\","
+					
+					"\"bool_int_0\" : true,"
+					"\"bool_int_1\" : false,"
+
+					"\"int_float_0\" : 10,"
+					"\"int_float_1\" : 100.00,"
+					"\"int_float_2\" : 10,"
+
+					"\"float_bool_0\" : 1.0,"
+					"\"float_bool_1\" : 0.0,"
+					"\"float_bool_2\" : 1.00000000000000000000000000000000000000000000000000000000000000000000000000000000000000000000000000001234,"
+					"\"float_bool_3\" : 0.00000000000000000000000000000000000000000000000000000000000000000000000000000000000000000000000000001234,"
+				"}"
+			);
+
+		const FString StoredAsType2 =
+			TEXT(
+				"{"
+					"\"bool_string_0\" : \"false\","
+					"\"bool_string_1\" : \"true\","
+					"\"bool_string_2\" : \"0\","
+					"\"bool_string_3\" : \"1\","
+
+					"\"int_string_0\" : \"10\","
+					"\"int_string_1\" : \"100\","
+
+					"\"float_string_0\" : \"10.123\","
+					"\"float_string_1\" : \"100.34\","
+					
+					"\"string_string_0\" : \"foo1\","
+					"\"string_string_1\" : \"foo2\","
+
+					"\"bool_int_0\" : 1,"
+					"\"bool_int_1\" : 0,"
+
+					"\"int_float_0\" : 10.0,"
+					"\"int_float_1\" : 100.00,"
+					"\"int_float_2\" : 10.00000000000000000000000000000000000000000000000000000000000000000000000000000000000000000000000000001234,"
+
+					"\"float_bool_0\" : true,"
+					"\"float_bool_1\" : false,"
+					"\"float_bool_2\" : true,"
+					"\"float_bool_3\" : false"
+				"}"
+			);
+
+		TSharedRef< TJsonReader<> > TypeReader_1 = TJsonReaderFactory<>::Create(StoredAsType1);
+		TSharedRef< TJsonReader<> > TypeReader_2 = TJsonReaderFactory<>::Create(StoredAsType2);
+
+		TSharedPtr<FJsonObject> TypedObject_1;
+		REQUIRE(FJsonSerializer::Deserialize(TypeReader_1, TypedObject_1));
+		REQUIRE(TypedObject_1.IsValid());
+
+		TSharedPtr<FJsonObject> TypedObject_2;
+		REQUIRE(FJsonSerializer::Deserialize(TypeReader_2, TypedObject_2));
+		REQUIRE(TypedObject_2.IsValid());
+
+		REQUIRE(TypedObject_1->Values.Num() == TypedObject_2->Values.Num());
+
+		for (const TPair<FString, TSharedPtr<FJsonValue>>& KV : TypedObject_1->Values)
+		{
+			REQUIRE(TypedObject_2->Values.Contains(KV.Key));
+			TSharedPtr<FJsonValue> Typed1_FieldValue = KV.Value;
+			TSharedPtr<FJsonValue> Typed2_FieldValue = TypedObject_2->Values[KV.Key];
+
+			REQUIRE(*Typed1_FieldValue == *Typed2_FieldValue);
+			REQUIRE(*Typed2_FieldValue == *Typed1_FieldValue);
+		}
+	}
+
+	// Variant Value Comparison Inequality Test
+	{
+		/* comparing: "Type1_Type2_#" */
+		const FString StoredAsType1 =
+			TEXT(
+				"{"
+					"\"bool_string_0\" : false,"
+					"\"bool_string_1\" : true,"
+
+					"\"int_string_0\" : 10,"
+					"\"int_string_1\" : 100,"
+
+					"\"float_string_0\" : 10.123,"
+					"\"float_string_1\" : 100.34,"
+
+					"\"string_string_0\" : \"foo1\","
+					"\"string_string_1\" : \"foo2\","
+
+					"\"bool_int_0\" : true,"
+					"\"bool_int_1\" : false,"
+
+					"\"int_float_0\" : 10,"
+					"\"int_float_1\" : 100.00,"
+					"\"int_float_2\" : 10,"
+
+					"\"float_bool_0\" : 1.0,"
+					"\"float_bool_1\" : 0.0,"
+					"\"float_bool_2\" : 2.5,"
+					"\"float_bool_3\" : 3.5,"
+				"}"
+			);
+
+		const FString StoredAsType2 =
+			TEXT(
+				"{"
+					"\"bool_string_0\" : \"not_true\","
+					"\"bool_string_1\" : \"not_false\","
+
+					"\"int_string_0\" : \"20\","
+					"\"int_string_1\" : \"200\","
+
+					"\"float_string_0\" : \"20.123\","
+					"\"float_string_1\" : \"200.34\","
+
+					"\"string_string_0\" : \"Foo1\","
+					"\"string_string_1\" : \"Foo2\","
+
+					"\"bool_int_0\" : 2,"
+					"\"bool_int_1\" : 3,"
+
+					"\"int_float_0\" : 20.0,"
+					"\"int_float_1\" : 200.00,"
+					"\"int_float_2\" : 10.5,"
+
+					"\"float_bool_0\" : false,"
+					"\"float_bool_1\" : true,"
+					"\"float_bool_2\" : true,"
+					"\"float_bool_3\" : false"
+				"}"
+			);
+
+		TSharedRef< TJsonReader<> > TypeReader_1 = TJsonReaderFactory<>::Create(StoredAsType1);
+		TSharedRef< TJsonReader<> > TypeReader_2 = TJsonReaderFactory<>::Create(StoredAsType2);
+
+		TSharedPtr<FJsonObject> TypedObject_1;
+		REQUIRE(FJsonSerializer::Deserialize(TypeReader_1, TypedObject_1));
+		REQUIRE(TypedObject_1.IsValid());
+
+		TSharedPtr<FJsonObject> TypedObject_2;
+		REQUIRE(FJsonSerializer::Deserialize(TypeReader_2, TypedObject_2));
+		REQUIRE(TypedObject_2.IsValid());
+
+		REQUIRE(TypedObject_1->Values.Num() == TypedObject_2->Values.Num());
+
+		for (const TPair<FString, TSharedPtr<FJsonValue>>& KV : TypedObject_1->Values)
+		{
+			REQUIRE(TypedObject_2->Values.Contains(KV.Key));
+			TSharedPtr<FJsonValue> Typed1_FieldValue = KV.Value;
+			TSharedPtr<FJsonValue> Typed2_FieldValue = TypedObject_2->Values[KV.Key];
+
+			REQUIRE(*Typed1_FieldValue != *Typed2_FieldValue);
+			REQUIRE(*Typed2_FieldValue != *Typed1_FieldValue);
+		}
+	}
+
 	// Number Array Case
 	{
 		const FString InputString =
diff --git a/Engine/Source/Runtime/Json/Public/Dom/JsonValue.h b/Engine/Source/Runtime/Json/Public/Dom/JsonValue.h
--- a/Engine/Source/Runtime/Json/Public/Dom/JsonValue.h
+++ b/Engine/Source/Runtime/Json/Public/Dom/JsonValue.h
@@ -250,3 +250,223 @@ public:
 protected:
 	virtual FString GetType() const override {return TEXT("Null");}
 };
+
+namespace UE::Json
+{
+
+template<typename T, typename = typename std::enable_if<!std::is_same_v<T, FJsonValue>>>
+static JsonSimpleValueVariant ToSimpleJsonVariant(const T& InSimpleValue)
+{
+	using InSimpleValueType = std::decay_t<decltype(InSimpleValue)>;
+	if constexpr (std::is_same_v<InSimpleValueType, bool> || std::is_same_v<InSimpleValueType, FString>)
+	{
+		return JsonSimpleValueVariant(TInPlaceType<T>(), InSimpleValue);
+	}
+	else
+	{
+		return JsonSimpleValueVariant(TInPlaceType<JsonNumberValueVariants>(), JsonNumberValueVariants(TInPlaceType<T>(), InSimpleValue));
+	}
+}
+
+static JsonSimpleValueVariant ToSimpleJsonVariant(const FJsonValue& InJsonValue)
+{
+	if (!InJsonValue.PreferStringRepresentation())
+	{
+		if (InJsonValue.Type == EJson::Boolean)
+		{
+			return JsonSimpleValueVariant(TInPlaceType<bool>(), InJsonValue.AsBool());
+		}
+		else if (InJsonValue.Type == EJson::Number)
+		{
+			const double JsonNumber_v = InJsonValue.AsNumber();
+
+			/* If the Json Number Value requires a decimal point, then we read in the value as a double, otherwise, we read it in as an int */
+			if (FString::SanitizeFloat(JsonNumber_v, 0).Contains(TEXT(".")))
+			{
+				return ToSimpleJsonVariant(JsonNumber_v);
+			}
+			else
+			{
+				return ToSimpleJsonVariant(FMath::RoundToInt64(JsonNumber_v));
+			}
+		}
+	}
+
+	return JsonSimpleValueVariant(TInPlaceType<FString>(), InJsonValue.AsString());
+}
+
+} // namespace UE::Json
+
+/* Global operators */
+
+static bool operator==(const JsonNumberValueVariants& Lhs, const FString& Rhs)
+{
+	return Rhs.IsNumeric() && ::Visit([Rhs](const auto& StoredNumber)
+		{
+			using StoredNumberType = std::decay_t<decltype(StoredNumber)>;
+			if constexpr (std::is_same_v<StoredNumberType, float> || std::is_same_v<StoredNumberType, double>)
+			{
+				return FString::SanitizeFloat(StoredNumber, 0) == Rhs;
+			}
+			else
+			{
+				return StoredNumber == FCString::Atoi64(*Rhs);
+			}
+		}, Lhs);
+}
+
+static bool operator!=(const JsonNumberValueVariants& Lhs, const FString& Rhs)
+{
+	return !(Lhs == Rhs);
+}
+
+static bool operator==(const FString& Lhs, const JsonNumberValueVariants& Rhs)
+{
+	return Rhs == Lhs;
+}
+
+static bool operator!=(const FString& Lhs, const JsonNumberValueVariants& Rhs)
+{
+	return !(Lhs == Rhs);
+}
+
+////////////////////////////////////////////////////////////
+
+static FString ToString(const JsonNumberValueVariants& InNumberVariant)
+{
+	return ::Visit([](auto& StoredNumber)
+		{
+			using StoredNumberType = std::decay_t<decltype(StoredNumber)>;
+			if constexpr (std::is_same_v<StoredNumberType, float> || std::is_same_v<StoredNumberType, double>)
+			{
+				return FString::SanitizeFloat(StoredNumber, 0);
+			}
+			else
+			{
+				return FString::Printf(TEXT("%lld"), static_cast<int64>(StoredNumber));
+			}
+		}, InNumberVariant);
+}
+
+static bool operator==(const JsonNumberValueVariants& Lhs, const JsonNumberValueVariants& Rhs)
+{
+	const bool bLhsIsFloat = Lhs.IsType<float>() || Lhs.IsType<double>();
+	const bool bRhsIsFloat = Rhs.IsType<float>() || Rhs.IsType<double>();
+	if (bLhsIsFloat || bRhsIsFloat)
+	{
+		return ToString(Lhs) == ToString(Rhs);
+	}
+	else
+	{
+		auto CastToInt64Functor = [](auto& StoredNumber)
+			{
+				return static_cast<int64>(StoredNumber);
+			};
+
+		const int64 LhsValue = ::Visit(CastToInt64Functor, Lhs);
+		const int64 RhsValue = ::Visit(CastToInt64Functor, Rhs);
+
+		return LhsValue == RhsValue;
+	}
+}
+
+static bool operator!=(const JsonNumberValueVariants& Lhs, const JsonNumberValueVariants& Rhs)
+{
+	return !(Lhs == Rhs);
+}
+
+////////////////////////////////////////////////////////////
+
+static bool operator==(const JsonSimpleValueVariant& Lhs, const JsonSimpleValueVariant& Rhs)
+{
+	if (Lhs.IsType<bool>())
+	{
+		if (Rhs.IsType<bool>())
+		{
+			return Lhs.Get<bool>() == Rhs.Get<bool>();
+		}
+		else if (Rhs.IsType<FString>())
+		{
+			if (Lhs.Get<bool>())
+			{
+				return Rhs.Get<FString>().Equals(TEXT("true"), ESearchCase::IgnoreCase) ||
+					Rhs.Get<FString>().Equals(TEXT("1"), ESearchCase::IgnoreCase);
+			}
+			else
+			{
+				return Rhs.Get<FString>().Equals(TEXT("false"), ESearchCase::IgnoreCase) ||
+					Rhs.Get<FString>().Equals(TEXT("0"), ESearchCase::IgnoreCase);
+			}
+		}
+		else // RhsType.IsType<JsonNumberValueVariants>()
+		{
+			return ::Visit([&Lhs](const auto& RhsStoredNumber)
+				{
+					using RhsStoredNumberType = std::decay_t<decltype(RhsStoredNumber)>;
+					if constexpr (std::is_same_v<RhsStoredNumberType, float> || std::is_same_v<RhsStoredNumberType, double>)
+					{
+						if (!FString::SanitizeFloat(RhsStoredNumber, 0).Contains(TEXT(".")))
+						{
+							const int64 RhsStoredNumberAsInt = FMath::RoundToInt64(RhsStoredNumber);
+							if (Lhs.Get<bool>())
+							{
+								return RhsStoredNumberAsInt == 1;
+							}
+							else
+							{
+								return RhsStoredNumberAsInt == 0;
+							}
+						}
+						else
+						{
+							return false;
+						}
+					}
+					else
+					{
+						if (Lhs.Get<bool>())
+						{
+							return RhsStoredNumber == 1;
+						}
+						else
+						{
+							return RhsStoredNumber == 0;
+						}
+					}
+				}, Rhs.Get<JsonNumberValueVariants>());
+		}
+	}
+	else if (Lhs.IsType<JsonNumberValueVariants>())
+	{
+		if (Rhs.IsType<JsonNumberValueVariants>())
+		{
+			return Lhs.Get<JsonNumberValueVariants>() == Rhs.Get<JsonNumberValueVariants>();
+		}
+		else // RhsType.IsType<bool>() || RhsType.IsType<FString>()
+		{
+			// Swapping args to avoid code duplication
+			return Rhs == Lhs;
+		}
+	}
+	else // Lhs.IsType<FString>()
+	{
+		if (Rhs.IsType<FString>())
+		{
+			return Lhs.Get<FString>() == Rhs.Get<FString>();
+		}
+		else if (Rhs.IsType<bool>())
+		{
+			// Swapping args to avoid code duplication
+			return Rhs == Lhs;
+		}
+		else // RhsType.IsType<JsonNumberValueVariants>()
+		{
+			return Lhs.Get<FString>() == Rhs.Get<JsonNumberValueVariants>();
+		}
+	}
+}
+
+static bool operator!=(const JsonSimpleValueVariant& Lhs, const JsonSimpleValueVariant& Rhs)
+{
+	return !(Lhs == Rhs);
+}
\ No newline at end of file
diff --git a/Engine/Source/Runtime/Json/Public/Serialization/JsonSerializerBase.h b/Engine/Source/Runtime/Json/Public/Serialization/JsonSerializerBase.h
--- a/Engine/Source/Runtime/Json/Public/Serialization/JsonSerializerBase.h
+++ b/Engine/Source/Runtime/Json/Public/Serialization/JsonSerializerBase.h
@@ -27,6 +27,7 @@ struct FJsonSerializerBase
 	virtual void Serialize(FStringView Name, float& Value) = 0;
 	virtual void Serialize(FStringView Name, double& Value) = 0;
 	virtual void Serialize(FStringView Name, FDateTime& Value) = 0;
+	virtual void Serialize(FStringView Name, JsonSimpleValueVariant& InVariant) = 0;
 	virtual void SerializeArray(FJsonSerializableArray& Array) = 0;
 	virtual void SerializeArray(FStringView Name, FJsonSerializableArray& Value) = 0;
 	virtual void SerializeArray(FStringView Name, FJsonSerializableArrayInt& Value) = 0;
@@ -36,6 +37,7 @@ struct FJsonSerializerBase
 	virtual void SerializeMap(FStringView Name, FJsonSerializableKeyValueMapArrayInt& Map) = 0;
 	virtual void SerializeMap(FStringView Name, FJsonSerializableKeyValueMapInt64& Map) = 0;
 	virtual void SerializeMap(FStringView Name, FJsonSerializableKeyValueMapFloat& Map) = 0;
+	virtual void SerializeMap(FStringView Name, FJsonSerializableKeySimpleValueVariantMap& Map) = 0;
 	virtual void SerializeSimpleMap(FJsonSerializableKeyValueMap& Map) = 0;
 	virtual void SerializeMapSafe(FStringView Name, FJsonSerializableKeyValueMap& Map) = 0;
 	virtual TSharedPtr<FJsonObject> GetObject() = 0;
diff --git a/Engine/Source/Runtime/Json/Public/Serialization/JsonSerializerMacros.h b/Engine/Source/Runtime/Json/Public/Serialization/JsonSerializerMacros.h
--- a/Engine/Source/Runtime/Json/Public/Serialization/JsonSerializerMacros.h
+++ b/Engine/Source/Runtime/Json/Public/Serialization/JsonSerializerMacros.h
@@ -56,6 +56,28 @@
 			} \
 		}
 
+#define JSON_SERIALIZE_OPTIONAL_MAP(JsonName, JsonMap) \
+	if (Serializer.IsLoading() && !JsonMap.IsSet() && Serializer.GetObject().IsValid() && Serializer.GetObject()->HasTypedField<EJson::Object>(TEXTVIEW(JsonName))) \
+	{ \
+		JsonMap.Emplace(); \
+	} \
+	\
+	if (JsonMap.IsSet()) \
+	{ \
+		Serializer.SerializeMap(TEXTVIEW(JsonName), JsonMap.GetValue()); \
+	}
+
+#define JSON_SERIALIZE_OPTIONAL_ARRAY(JsonName, JsonArray) \
+	if (Serializer.IsLoading() && !JsonArray.IsSet() && Serializer.GetObject().IsValid() && Serializer.GetObject()->HasTypedField<EJson::Array>(TEXTVIEW(JsonName))) \
+	{ \
+		JsonArray.Emplace(); \
+	} \
+	\
+	if (JsonArray.IsSet()) \
+	{ \
+		Serializer.SerializeArray(TEXTVIEW(JsonName), JsonArray.GetValue()); \
+	}
+
 #define JSON_SERIALIZE_ARRAY(JsonName, JsonArray) \
 		Serializer.SerializeArray(TEXTVIEW(JsonName), JsonArray)
 		
diff --git a/Engine/Source/Runtime/Json/Public/Serialization/JsonSerializerReader.h b/Engine/Source/Runtime/Json/Public/Serialization/JsonSerializerReader.h
--- a/Engine/Source/Runtime/Json/Public/Serialization/JsonSerializerReader.h
+++ b/Engine/Source/Runtime/Json/Public/Serialization/JsonSerializerReader.h
@@ -102,6 +102,14 @@ public:
 	* @param Value the value to write out
 	*/
 	JSON_API virtual void Serialize(FStringView Name, FDateTime& Value) override;
+	/**
+	* Writes the field name and the corresponding value to the JSON data
+	*
+	* @param Name the field name to write out
+	* @param Value the value to write out
+	*/
+	JSON_API virtual void Serialize(FStringView Name, JsonSimpleValueVariant& InVariant) override;
+
 	/**
 	 * Serializes an array of values
 	 *
@@ -170,6 +178,14 @@ public:
 	 */
 	JSON_API virtual void SerializeMap(FStringView Name, FJsonSerializableKeyValueMapFloat& Map) override;
 
+	/**
+	 * Serializes the keys & values for map
+	 *
+	 * @param Name the name of the property to serialize
+	 * @param Map the map to serialize
+	 */
+	JSON_API virtual void SerializeMap(FStringView Name, FJsonSerializableKeySimpleValueVariantMap& Map) override;
+
 	JSON_API virtual void SerializeSimpleMap(FJsonSerializableKeyValueMap& Map) override;
 
 	/**
diff --git a/Engine/Source/Runtime/Json/Public/Serialization/JsonSerializerWriter.h b/Engine/Source/Runtime/Json/Public/Serialization/JsonSerializerWriter.h
--- a/Engine/Source/Runtime/Json/Public/Serialization/JsonSerializerWriter.h
+++ b/Engine/Source/Runtime/Json/Public/Serialization/JsonSerializerWriter.h
@@ -170,6 +170,26 @@ public:
 			JsonWriter->WriteValue(Name, Value.ToIso8601());
 		}
 	}
+
+	virtual void Serialize(FStringView Name, JsonSimpleValueVariant& InVariant) override
+	{
+		::Visit([this, &Name](auto& StoredValue)
+			{
+				using StoredValueType = std::decay_t<decltype(StoredValue)>;
+				if constexpr (std::is_same_v<StoredValueType, JsonNumberValueVariants>)
+				{
+					::Visit([this, &Name](auto& StoredNumberValue)
+						{
+							JsonWriter->WriteValue(Name, StoredNumberValue);
+						}, StoredValue);
+				}
+				else
+				{
+					JsonWriter->WriteValue(Name, StoredValue);
+				}
+			}, InVariant);
+	}
+
 	/**
 	 * Serializes an array of values
 	 *
@@ -320,6 +340,31 @@ public:
 		JsonWriter->WriteObjectEnd();
 	}
 
+	virtual void SerializeMap(FStringView Name, FJsonSerializableKeySimpleValueVariantMap& Map) override
+	{
+		JsonWriter->WriteObjectStart(Name);
+		// Iterate all of the keys and their values
+		for (TPair<FString, JsonSimpleValueVariant>& Pair : Map)
+		{
+			::Visit([this, &Pair](auto& StoredValue)
+				{
+					using StoredValueType = std::decay_t<decltype(StoredValue)>;
+					if constexpr (std::is_same_v<StoredValueType, JsonNumberValueVariants>)
+					{
+						::Visit([this, &Pair](auto& StoredNumberValue)
+							{
+								Serialize(*Pair.Key, StoredNumberValue);
+							}, StoredValue);
+					}
+					else
+					{
+						Serialize(*Pair.Key, StoredValue);
+					}
+				}, Pair.Value);
+		}
+		JsonWriter->WriteObjectEnd();
+	}
+
 	virtual void SerializeSimpleMap(FJsonSerializableKeyValueMap& Map) override
 	{
 		// writing does nothing here, this is meant to read in all data from a json object 
diff --git a/Engine/Source/Runtime/Json/Public/Serialization/JsonTypes.h b/Engine/Source/Runtime/Json/Public/Serialization/JsonTypes.h
--- a/Engine/Source/Runtime/Json/Public/Serialization/JsonTypes.h
+++ b/Engine/Source/Runtime/Json/Public/Serialization/JsonTypes.h
@@ -3,6 +3,7 @@
 #pragma once
 
 #include "CoreMinimal.h"
+#include "Misc/TVariant.h"
 
 class Error;
 
@@ -80,3 +81,10 @@ using FJsonSerializableKeyValueMapInt64 = TMap<FString, int64>;
 using FJsonSerializableKeyValueMapFloat = TMap<FString, float>;
 using FJsonSerializableKeyValueMapArrayInt = TMap<FString, FJsonSerializableArrayInt>;
 
+/** We always read with the highest precisions (int64 or double), but we can write with any level of precision */
+using JsonNumberValueVariants = TVariant<int32, uint32, int64, float, double>; 
+
+using JsonSimpleValueVariant = TVariant<bool /* EJson::Boolean */, JsonNumberValueVariants /* EJson::Number */, FString /* EJson::String */>;
+
+/** Helps keep key values unique */
+using FJsonSerializableKeySimpleValueVariantMap = TMap<FString, JsonSimpleValueVariant>;
\ No newline at end of file
