Parse fix

diff --git a/Engine/Source/Runtime/Json/Private/Tests/JsonTests.cpp b/Engine/Source/Runtime/Json/Private/Tests/JsonTests.cpp
--- a/Engine/Source/Runtime/Json/Private/Tests/JsonTests.cpp
+++ b/Engine/Source/Runtime/Json/Private/Tests/JsonTests.cpp
@@ -191,7 +191,7 @@ TEST_CASE_NAMED(FJsonAutomationTest, "System::Engine::FileSystem::JSON", "[Appli
 					"\"float_bool_0\" : 1.0,"
 					"\"float_bool_1\" : 0.0,"
 					"\"float_bool_2\" : 1.00000000000000000000000000000000000000000000000000000000000000000000000000000000000000000000000000001234,"
-					"\"float_bool_3\" : 0.00000000000000000000000000000000000000000000000000000000000000000000000000000000000000000000000000001234,"
+					"\"float_bool_3\" : 0.00000000000000000000000000000000000000000000000000000000000000000000000000000000000000000000000000001234"
 				"}"
 			);
 
@@ -278,7 +278,7 @@ TEST_CASE_NAMED(FJsonAutomationTest, "System::Engine::FileSystem::JSON", "[Appli
 					"\"float_bool_0\" : 1.0,"
 					"\"float_bool_1\" : 0.0,"
 					"\"float_bool_2\" : 2.5,"
-					"\"float_bool_3\" : 3.5,"
+					"\"float_bool_3\" : 3.5"
 				"}"
 			);
 
