Removing String to String test

diff --git a/Engine/Source/Runtime/Json/Private/Tests/JsonTests.cpp b/Engine/Source/Runtime/Json/Private/Tests/JsonTests.cpp
--- a/Engine/Source/Runtime/Json/Private/Tests/JsonTests.cpp
+++ b/Engine/Source/Runtime/Json/Private/Tests/JsonTests.cpp
@@ -265,9 +265,6 @@ TEST_CASE_NAMED(FJsonAutomationTest, "System::Engine::FileSystem::JSON", "[Appli
 					"\"float_string_0\" : 10.123,"
 					"\"float_string_1\" : 100.34,"
 
-					"\"string_string_0\" : \"foo1\","
-					"\"string_string_1\" : \"foo2\","
-
 					"\"bool_int_0\" : true,"
 					"\"bool_int_1\" : false,"
 
@@ -294,9 +291,6 @@ TEST_CASE_NAMED(FJsonAutomationTest, "System::Engine::FileSystem::JSON", "[Appli
 					"\"float_string_0\" : \"20.123\","
 					"\"float_string_1\" : \"200.34\","
 
-					"\"string_string_0\" : \"Foo1\","
-					"\"string_string_1\" : \"Foo2\","
-
 					"\"bool_int_0\" : 2,"
 					"\"bool_int_1\" : 3,"
 
