Added more Json Tests

diff --git a/Engine/Source/Runtime/Json/Private/Tests/JsonTests.cpp b/Engine/Source/Runtime/Json/Private/Tests/JsonTests.cpp
--- a/Engine/Source/Runtime/Json/Private/Tests/JsonTests.cpp
+++ b/Engine/Source/Runtime/Json/Private/Tests/JsonTests.cpp
@@ -161,7 +161,7 @@ TEST_CASE_NAMED(FJsonAutomationTest, "System::Engine::FileSystem::JSON", "[Appli
 		REQUIRE(InputString == OutputString);
 	}
 
-	// Variant Value Comparison Equality Test
+	// FJsonValue operator== Comparison Equality Test
 	{
 		/* comparing: "Type1_Type2_#" */
 		const FString StoredAsType1 =
@@ -250,7 +250,7 @@ TEST_CASE_NAMED(FJsonAutomationTest, "System::Engine::FileSystem::JSON", "[Appli
 		}
 	}
 
-	// Variant Value Comparison Inequality Test
+	// FJsonValue operator!= Comparison Inequality Test
 	{
 		/* comparing: "Type1_Type2_#" */
 		const FString StoredAsType1 =
@@ -265,6 +265,10 @@ TEST_CASE_NAMED(FJsonAutomationTest, "System::Engine::FileSystem::JSON", "[Appli
 					"\"float_string_0\" : 10.123,"
 					"\"float_string_1\" : 100.34,"
 
+					// `FJsonValue operator==` uses `FString::operator==` which uses `ESearchCase::IgnoreCase`
+					//"\"string_string_0\" : \"foo1\","
+					//"\"string_string_1\" : \"foo2\","
+
 					"\"bool_int_0\" : true,"
 					"\"bool_int_1\" : false,"
 
@@ -291,6 +295,10 @@ TEST_CASE_NAMED(FJsonAutomationTest, "System::Engine::FileSystem::JSON", "[Appli
 					"\"float_string_0\" : \"20.123\","
 					"\"float_string_1\" : \"200.34\","
 
+					// `FJsonValue operator==` uses `FString::operator==` which uses `ESearchCase::IgnoreCase`
+					//"\"string_string_0\" : \"Foo1\","
+					//"\"string_string_1\" : \"Foo2\","
+
 					"\"bool_int_0\" : 2,"
 					"\"bool_int_1\" : 3,"
 
@@ -329,6 +337,182 @@ TEST_CASE_NAMED(FJsonAutomationTest, "System::Engine::FileSystem::JSON", "[Appli
 		}
 	}
 
+	// JsonSimpleValueVariant operator== Comparison Equality Test
+	{
+		/* comparing: "Type1_Type2_#" */
+		const FString StoredAsType1 =
+			TEXT(
+				"{"
+				"\"bool_string_0\" : false,"
+				"\"bool_string_1\" : true,"
+				"\"bool_string_2\" : false,"
+				"\"bool_string_3\" : true,"
+
+				"\"int_string_0\" : 10,"
+				"\"int_string_1\" : 100,"
+
+				"\"float_string_0\" : 10.123,"
+				"\"float_string_1\" : 100.34,"
+
+				"\"string_string_0\" : \"foo1\","
+				"\"string_string_1\" : \"foo2\","
+
+				"\"bool_int_0\" : true,"
+				"\"bool_int_1\" : false,"
+
+				"\"int_float_0\" : 10,"
+				"\"int_float_1\" : 100.00,"
+				"\"int_float_2\" : 10,"
+
+				"\"float_bool_0\" : 1.0,"
+				"\"float_bool_1\" : 0.0,"
+				"\"float_bool_2\" : 1.00000000000000000000000000000000000000000000000000000000000000000000000000000000000000000000000000001234,"
+				"\"float_bool_3\" : 0.00000000000000000000000000000000000000000000000000000000000000000000000000000000000000000000000000001234"
+				"\"float_bool_4\" : 0.9999999999999999999999999999999999999999999999999999999999999999999999999876,"
+				"}"
+			);
+
+		const FString StoredAsType2 =
+			TEXT(
+				"{"
+				"\"bool_string_0\" : \"false\","
+				"\"bool_string_1\" : \"true\","
+				"\"bool_string_2\" : \"0\","
+				"\"bool_string_3\" : \"1\","
+
+				"\"int_string_0\" : \"10\","
+				"\"int_string_1\" : \"100\","
+
+				"\"float_string_0\" : \"10.123\","
+				"\"float_string_1\" : \"100.34\","
+
+				"\"string_string_0\" : \"foo1\","
+				"\"string_string_1\" : \"foo2\","
+
+				"\"bool_int_0\" : 1,"
+				"\"bool_int_1\" : 0,"
+
+				"\"int_float_0\" : 10.0,"
+				"\"int_float_1\" : 100.00,"
+				"\"int_float_2\" : 10.00000000000000000000000000000000000000000000000000000000000000000000000000000000000000000000000000001234,"
+
+				"\"float_bool_0\" : true,"
+				"\"float_bool_1\" : false,"
+				"\"float_bool_2\" : true,"
+				"\"float_bool_3\" : false,"
+				"\"float_bool_4\" : true"
+				"}"
+			);
+
+		TSharedRef< TJsonReader<> > TypeReader_1 = TJsonReaderFactory<>::Create(StoredAsType1);
+		TSharedRef< TJsonReader<> > TypeReader_2 = TJsonReaderFactory<>::Create(StoredAsType2);
+
+		TSharedPtr<FJsonObject> TypedObject_1;
+		REQUIRE(FJsonSerializer::Deserialize(TypeReader_1, TypedObject_1));
+		REQUIRE(TypedObject_1.IsValid());
+
+		TSharedPtr<FJsonObject> TypedObject_2;
+		REQUIRE(FJsonSerializer::Deserialize(TypeReader_2, TypedObject_2));
+		REQUIRE(TypedObject_2.IsValid());
+
+		REQUIRE(TypedObject_1->Values.Num() == TypedObject_2->Values.Num());
+
+		for (const TPair<FString, TSharedPtr<FJsonValue>>& KV : TypedObject_1->Values)
+		{
+			REQUIRE(TypedObject_2->Values.Contains(KV.Key));
+			TSharedPtr<FJsonValue> Typed1_FieldValue = KV.Value;
+			TSharedPtr<FJsonValue> Typed2_FieldValue = TypedObject_2->Values[KV.Key];
+
+			REQUIRE(UE::Json::ToSimpleJsonVariant(*Typed1_FieldValue) == UE::Json::ToSimpleJsonVariant(*Typed2_FieldValue));
+			REQUIRE(UE::Json::ToSimpleJsonVariant(*Typed2_FieldValue) == UE::Json::ToSimpleJsonVariant(*Typed1_FieldValue));
+		}
+	}
+
+	// JsonSimpleValueVariant operator!= Comparison Inequality Test
+	{
+		/* comparing: "Type1_Type2_#" */
+		const FString StoredAsType1 =
+			TEXT(
+				"{"
+				"\"bool_string_0\" : false,"
+				"\"bool_string_1\" : true,"
+
+				"\"int_string_0\" : 10,"
+				"\"int_string_1\" : 100,"
+
+				"\"float_string_0\" : 10.123,"
+				"\"float_string_1\" : 100.34,"
+
+				"\"string_string_0\" : \"foo1\","
+				"\"string_string_1\" : \"foo2\","
+
+				"\"bool_int_0\" : true,"
+				"\"bool_int_1\" : false,"
+
+				"\"int_float_0\" : 10,"
+				"\"int_float_1\" : 100.00,"
+				"\"int_float_2\" : 10,"
+
+				"\"float_bool_0\" : 1.0,"
+				"\"float_bool_1\" : 0.0,"
+				"\"float_bool_2\" : 2.5,"
+				"\"float_bool_3\" : 3.5"
+				"}"
+			);
+
+		const FString StoredAsType2 =
+			TEXT(
+				"{"
+				"\"bool_string_0\" : \"not_true\","
+				"\"bool_string_1\" : \"not_false\","
+
+				"\"int_string_0\" : \"20\","
+				"\"int_string_1\" : \"200\","
+
+				"\"float_string_0\" : \"20.123\","
+				"\"float_string_1\" : \"200.34\","
+
+				"\"string_string_0\" : \"Foo1\","
+				"\"string_string_1\" : \"Foo2\","
+
+				"\"bool_int_0\" : 2,"
+				"\"bool_int_1\" : 3,"
+
+				"\"int_float_0\" : 20.0,"
+				"\"int_float_1\" : 200.00,"
+				"\"int_float_2\" : 10.5,"
+
+				"\"float_bool_0\" : false,"
+				"\"float_bool_1\" : true,"
+				"\"float_bool_2\" : true,"
+				"\"float_bool_3\" : false"
+				"}"
+			);
+
+		TSharedRef< TJsonReader<> > TypeReader_1 = TJsonReaderFactory<>::Create(StoredAsType1);
+		TSharedRef< TJsonReader<> > TypeReader_2 = TJsonReaderFactory<>::Create(StoredAsType2);
+
+		TSharedPtr<FJsonObject> TypedObject_1;
+		REQUIRE(FJsonSerializer::Deserialize(TypeReader_1, TypedObject_1));
+		REQUIRE(TypedObject_1.IsValid());
+
+		TSharedPtr<FJsonObject> TypedObject_2;
+		REQUIRE(FJsonSerializer::Deserialize(TypeReader_2, TypedObject_2));
+		REQUIRE(TypedObject_2.IsValid());
+
+		REQUIRE(TypedObject_1->Values.Num() == TypedObject_2->Values.Num());
+
+		for (const TPair<FString, TSharedPtr<FJsonValue>>& KV : TypedObject_1->Values)
+		{
+			REQUIRE(TypedObject_2->Values.Contains(KV.Key));
+			TSharedPtr<FJsonValue> Typed1_FieldValue = KV.Value;
+			TSharedPtr<FJsonValue> Typed2_FieldValue = TypedObject_2->Values[KV.Key];
+
+			REQUIRE(UE::Json::ToSimpleJsonVariant(*Typed1_FieldValue) != UE::Json::ToSimpleJsonVariant(*Typed2_FieldValue));
+			REQUIRE(UE::Json::ToSimpleJsonVariant(*Typed2_FieldValue) != UE::Json::ToSimpleJsonVariant(*Typed1_FieldValue));
+		}
+	}
+
 	// Number Array Case
 	{
 		const FString InputString =
