Allow JSON Array serialization to be an array of elements or array of pointers

diff --git a/Engine/Source/Runtime/Json/Public/Serialization/JsonSerializable.h b/Engine/Source/Runtime/Json/Public/Serialization/JsonSerializable.h
--- a/Engine/Source/Runtime/Json/Public/Serialization/JsonSerializable.h
+++ b/Engine/Source/Runtime/Json/Public/Serialization/JsonSerializable.h
@@ -125,6 +125,24 @@ using ToJsonVariantArgs = TVariant<ReturnStringArgs, WriterVariants>;
 using PrettySerializer = FJsonSerializerWriter<>;
 using CondensedSerializer = FJsonSerializerWriter<TCHAR, TCondensedJsonPrintPolicy<TCHAR>>;
 
+template<typename T, typename...SerializerArgsT>
+inline void ToJson_SerializeArrayElements(TArray<T>& InArray, SerializerArgsT...Args)
+{
+	for (T& ArrayEntry : InArray)
+	{
+		ArrayEntry.Serialize(Args...);
+	}
+}
+
+template<typename T, typename...SerializerArgsT>
+inline void ToJson_SerializeArrayElements(TArray<T*>& InArray, SerializerArgsT...Args)
+{
+	for (T* ArrayEntry : InArray)
+	{
+		ArrayEntry->Serialize(Args...);
+	}
+}
+
 template<typename T>
 inline void ToJson(TArray<T>& InArray, const ToJsonVariantArgs& InArgs)
 {
@@ -173,10 +191,7 @@ inline void ToJson(TArray<T>& InArray, const ToJsonVariantArgs& InArgs)
 		{
 			StoredSerializer.template Get<0>().StartArray();
 
-			for (T& ArrayEntry : InArray)
-			{
-				ArrayEntry.Serialize(StoredSerializer.template Get<0>(), false);
-			}
+			ToJson_SerializeArrayElements(InArray, StoredSerializer.template Get<0>(), false);
 
 			StoredSerializer.template Get<0>().EndArray();
 
