Make FJsonSerializable::ToJson generic so it works with all json writers

diff --git a/Engine/Source/Runtime/Json/Private/Serialization/JsonSerializable.cpp b/Engine/Source/Runtime/Json/Private/Serialization/JsonSerializable.cpp
--- a/Engine/Source/Runtime/Json/Private/Serialization/JsonSerializable.cpp
+++ b/Engine/Source/Runtime/Json/Private/Serialization/JsonSerializable.cpp
@@ -35,18 +35,6 @@ const FString FJsonSerializable::ToJson(bool bPrettyPrint/*=true*/)
 	return JsonStr;
 }
 
-void FJsonSerializable::ToJson(TSharedRef<TJsonWriter<> >& JsonWriter, bool bFlatObject) const
-{
-	FJsonSerializerWriter<> Serializer(JsonWriter);
-	((FJsonSerializable*)this)->Serialize(Serializer, bFlatObject);
-}
-
-void FJsonSerializable::ToJson(TSharedRef< TJsonWriter< TCHAR, TCondensedJsonPrintPolicy< TCHAR > > >& JsonWriter, bool bFlatObject) const
-{
-	FJsonSerializerWriter<TCHAR, TCondensedJsonPrintPolicy< TCHAR >> Serializer(JsonWriter);
-	((FJsonSerializable*)this)->Serialize(Serializer, bFlatObject);
-}
-
 bool FJsonSerializable::FromJson(const FString& Json)
 {
 	return FromJsonStringView(FStringView(Json));
diff --git a/Engine/Source/Runtime/Json/Public/Serialization/JsonSerializable.h b/Engine/Source/Runtime/Json/Public/Serialization/JsonSerializable.h
--- a/Engine/Source/Runtime/Json/Public/Serialization/JsonSerializable.h
+++ b/Engine/Source/Runtime/Json/Public/Serialization/JsonSerializable.h
@@ -26,6 +26,7 @@ struct FJsonSerializable
 	 * @return the corresponding json string
 	 */
 	JSON_API const FString ToJson(bool bPrettyPrint = true) const;
+
 	/**
 	 * Serializes this object to its JSON string form
 	 *
@@ -33,8 +34,15 @@ struct FJsonSerializable
 	 * @return the corresponding json string
 	 */
 	JSON_API virtual const FString ToJson(bool bPrettyPrint=true);
-	JSON_API virtual void ToJson(TSharedRef<TJsonWriter<> >& JsonWriter, bool bFlatObject) const;
-	JSON_API virtual void ToJson(TSharedRef< TJsonWriter< TCHAR, TCondensedJsonPrintPolicy< TCHAR > > >& JsonWriter, bool bFlatObject) const;
+
+	/**
+	 * Serializes this object with a Json Writer
+	 * 
+	 * @param JsonWriter - The writer to use
+	 * @param bFlatObject if true then no object wrapper is used
+	 */
+	template<class CharType, class PrintPolicy, ESPMode SPMode>
+	void ToJson(TSharedRef<TJsonWriter<CharType, PrintPolicy>, SPMode> JsonWriter, bool bFlatObject = false) const;
 
 	/**
 	 * Serializes the contents of a JSON string into this object
@@ -75,6 +83,13 @@ struct FJsonSerializable
 	JSON_API virtual void Serialize(FJsonSerializerBase& Serializer, bool bFlatObject) = 0;
 };
 
+template<class CharType, class PrintPolicy, ESPMode SPMode>
+inline void FJsonSerializable::ToJson(TSharedRef<TJsonWriter<CharType, PrintPolicy>, SPMode> JsonWriter, bool bFlatObject) const
+{
+	FJsonSerializerWriter<CharType, PrintPolicy> Serializer(MoveTemp(JsonWriter));
+	const_cast<FJsonSerializable*>(this)->Serialize(Serializer, bFlatObject);
+}
+
 namespace UE::JsonArray
 {
 
@@ -248,4 +263,4 @@ static void ToJson(TArray<T>& InArray, Private::CondensedWriter& JsonWriter)
 	Private::ToJson(InArray, Private::ToJsonVariantArgs(TInPlaceType<Private::WriterVariants>(), Private::WriterVariants(TInPlaceType<Private::CondensedWriter>(), JsonWriter)));
 }
 
-} // namespace UE::JsonArray
\ No newline at end of file
+} // namespace UE::JsonArray
