Support reading nan in JsonReader(case insensitive).

diff --git a/Engine/Source/Runtime/Json/Private/Tests/JsonTests.cpp b/Engine/Source/Runtime/Json/Private/Tests/JsonTests.cpp
--- a/Engine/Source/Runtime/Json/Private/Tests/JsonTests.cpp
+++ b/Engine/Source/Runtime/Json/Private/Tests/JsonTests.cpp
@@ -802,6 +802,45 @@ TEST_CASE_NAMED(FJsonAutomationTest, "System::Engine::FileSystem::JSON", "[Appli
 		REQUIRE(OutputString == TestOutput);
 	}
 
+	// Test Nan
+	{
+		const FString InputString =
+			TEXT(
+				"{"
+					"\"Value0\":nan,"
+					"\"Value1\":NaN"
+				"}"
+			);
+		TSharedRef< TJsonReader<> > Reader = TJsonReaderFactory<>::Create( InputString );
+
+		TSharedPtr<FJsonObject> Object;
+		bool bSuccessful = FJsonSerializer::Deserialize(Reader, Object);
+		REQUIRE(bSuccessful);
+		REQUIRE(Object.IsValid());
+
+		const TSharedPtr<FJsonValue>* Value0 = Object->Values.Find(FString::Printf(TEXT("Value%i"), 0));
+		const TSharedPtr<FJsonValue>* Value1 = Object->Values.Find(FString::Printf(TEXT("Value%i"), 1));
+		const double Number0 = (*Value0)->AsNumber();
+		const double Number1 = (*Value1)->AsNumber();
+		REQUIRE(FMath::IsNaN(Number0));
+		REQUIRE(FMath::IsNaN(Number1));
+
+		FString OutputString;
+		TSharedRef< FCondensedJsonStringWriter > Writer = FCondensedJsonStringWriterFactory::Create( &OutputString );
+		REQUIRE(FJsonSerializer::Serialize( Object.ToSharedRef(), Writer ));
+
+		// %g isn't standardized, so we use the same %g format that is used inside PrintJson instead of hardcoding the values here
+		const FString TestOutput = FString::Printf(
+			TEXT(
+				"{"
+					"\"Value0\":%.17g,"
+					"\"Value1\":%.17g"
+				"}"
+			),
+			std::numeric_limits<double>::quiet_NaN(), std::numeric_limits<double>::quiet_NaN());
+		REQUIRE(OutputString == TestOutput);
+	}
+
 	// Boolean/Null Test
 	{
 		const FString InputString =
diff --git a/Engine/Source/Runtime/Json/Public/Serialization/JsonReader.h b/Engine/Source/Runtime/Json/Public/Serialization/JsonReader.h
--- a/Engine/Source/Runtime/Json/Public/Serialization/JsonReader.h
+++ b/Engine/Source/Runtime/Json/Public/Serialization/JsonReader.h
@@ -537,6 +537,13 @@ private:
 							return true;
 						}
 
+						if (Test.Compare(TEXT("NaN"), ESearchCase::IgnoreCase) == 0)
+						{
+							NumberValue = std::numeric_limits<double>::quiet_NaN();
+							OutToken = EJsonToken::Number;
+							return true;
+						}
+
 						SetErrorMessage( TEXT("Invalid Json Token. Check that your member names have quotes around them!") );
 						return false;
 					}
