Added methods to assist with tracking memory associated with JsonValue, JsonObject and JsonObjectWrapper

* Added GetMemoryFootprint() implementation which handles returning memory used by the different objects, this returns memory for sizeof(*this*) and size of any memory dynamically allocated by the object instance
* Added GetAllocatedSize() implementation which is meant to only return size of any memory dynamically allocated by the object instance

#rnx

diff --git a/Engine/Source/Runtime/Json/Private/Dom/JsonObject.cpp b/Engine/Source/Runtime/Json/Private/Dom/JsonObject.cpp
--- a/Engine/Source/Runtime/Json/Private/Dom/JsonObject.cpp
+++ b/Engine/Source/Runtime/Json/Private/Dom/JsonObject.cpp
@@ -234,6 +234,20 @@ void FJsonObject::SetObjectField(const FString& FieldName, const TSharedPtr<FJso
 	SetObjectField(CopyTemp(FieldName), JsonObject);
 }
 
+SIZE_T FJsonObject::GetAllocatedSize() const
+{
+	SIZE_T SizeBytes = 0;
+
+	SizeBytes += Values.GetAllocatedSize();
+	for (const TPair<FString, TSharedPtr<FJsonValue>>& KV : Values)
+	{
+		SizeBytes += KV.Key.GetAllocatedSize();
+		SizeBytes += KV.Value.IsValid() ? KV.Value->GetMemoryFootprint() : 0;
+	}
+
+	return SizeBytes;
+}
+
 void FJsonObject::Duplicate(const TSharedPtr<const FJsonObject>& Source, const TSharedPtr<FJsonObject>& Dest)
 {
 	if (Source && Dest)
diff --git a/Engine/Source/Runtime/Json/Private/Dom/JsonValue.cpp b/Engine/Source/Runtime/Json/Private/Dom/JsonValue.cpp
--- a/Engine/Source/Runtime/Json/Private/Dom/JsonValue.cpp
+++ b/Engine/Source/Runtime/Json/Private/Dom/JsonValue.cpp
@@ -353,3 +353,8 @@ void FJsonValue::ErrorMessage(const FString& InType) const
 		UE_LOG(LogJson, Error, TEXT("Json Value of type '%s' used as a '%s'."), *GetType(), *InType);
 	}
 }
+
+SIZE_T FJsonValueObject::GetAllocatedSize() const
+{
+	return Value.IsValid() ? Value->GetMemoryFootprint() : 0;
+}
diff --git a/Engine/Source/Runtime/Json/Public/Dom/JsonObject.h b/Engine/Source/Runtime/Json/Public/Dom/JsonObject.h
--- a/Engine/Source/Runtime/Json/Public/Dom/JsonObject.h
+++ b/Engine/Source/Runtime/Json/Public/Dom/JsonObject.h
@@ -264,6 +264,15 @@ public:
 	JSON_API void SetObjectField(FString&& FieldName, const TSharedPtr<FJsonObject>& JsonObject);
 	JSON_API void SetObjectField(const FString& FieldName, const TSharedPtr<FJsonObject>& JsonObject);
 
+	/**
+	 * Returns the memory footprint for this object in Bytes, including sizeof(*this) and allocated memory.
+	 */
+	SIZE_T GetMemoryFootprint() const { return sizeof(*this) + GetAllocatedSize(); }
+private:
+	/** Helper to calculate allocated size of the Values map and its contents */
+	JSON_API SIZE_T GetAllocatedSize() const;
+
+public:
 	static JSON_API void Duplicate(const TSharedPtr<const FJsonObject>& Source, const TSharedPtr<FJsonObject>& Dest);
 	static JSON_API void Duplicate(const TSharedPtr<FJsonObject>& Source, TSharedPtr<FJsonObject>& Dest);
 
diff --git a/Engine/Source/Runtime/Json/Public/Dom/JsonValue.h b/Engine/Source/Runtime/Json/Public/Dom/JsonValue.h
--- a/Engine/Source/Runtime/Json/Public/Dom/JsonValue.h
+++ b/Engine/Source/Runtime/Json/Public/Dom/JsonValue.h
@@ -96,6 +96,12 @@ public:
 	void AsArgumentType(TArray< TSharedPtr<FJsonValue> >& Value) { Value = AsArray (); }
 	void AsArgumentType(TSharedPtr<FJsonObject>         & Value) { Value = AsObject(); }
 
+	/**
+	 * Returns the memory footprint for this object in Bytes, including sizeof(*this) and allocated memory.
+	 * All children should implement this so their memory layout is properly accounted for
+	 */
+	JSON_API virtual SIZE_T GetMemoryFootprint() const { return sizeof(*this); }
+
 	EJson Type;
 
 	static JSON_API TSharedPtr<FJsonValue> Duplicate(const TSharedPtr<const FJsonValue>& Src);
@@ -142,10 +148,14 @@ public:
 	// Way to check if string value is empty without copying the string 
 	bool IsEmpty() const { return Value.IsEmpty(); }
 
+	virtual SIZE_T GetMemoryFootprint() const override { return sizeof(*this) + GetAllocatedSize(); }
+
 protected:
 	FString Value;
 
 	virtual FString GetType() const override {return TEXT("String");}
+	/** Helper to calculate allocated size of the Value string */
+	SIZE_T GetAllocatedSize() const { return Value.GetAllocatedSize(); }
 };
 
 
@@ -157,7 +167,8 @@ public:
 	virtual bool TryGetNumber(double& OutNumber) const override		{ OutNumber = Value; return true; }
 	virtual bool TryGetBool(bool& OutBool) const override			{ OutBool = (Value != 0.0); return true; }
 	virtual bool TryGetString(FString& OutString) const override	{ OutString = FString::SanitizeFloat(Value, 0); return true; }
-	
+	virtual SIZE_T GetMemoryFootprint() const override { return sizeof(*this); }
+
 protected:
 
 	double Value;
@@ -186,11 +197,14 @@ public:
 	virtual bool TryGetNumber(uint64& OutValue) const override { return LexTryParseString(OutValue, *Value); }
 	virtual bool TryGetBool(bool& OutBool) const override { OutBool = Value.ToBool(); return true; }
 	virtual bool PreferStringRepresentation() const override { return true; }
+	virtual SIZE_T GetMemoryFootprint() const override { return sizeof(*this) + GetAllocatedSize(); }
 
 protected:
 	FString Value;
 
 	virtual FString GetType() const override { return TEXT("NumberString"); }
+	/** Helper to calculate allocated size of the Value string */
+	SIZE_T GetAllocatedSize() const { return Value.GetAllocatedSize(); }
 };
 
 
@@ -202,7 +216,8 @@ public:
 	virtual bool TryGetNumber(double& OutNumber) const override		{ OutNumber = Value ? 1 : 0; return true; }
 	virtual bool TryGetBool(bool& OutBool) const override			{ OutBool = Value; return true; }
 	virtual bool TryGetString(FString& OutString) const override	{ OutString = Value ? TEXT("true") : TEXT("false"); return true; }
-	
+	virtual SIZE_T GetMemoryFootprint() const override { return sizeof(*this); }
+
 protected:
 	bool Value;
 
@@ -218,11 +233,23 @@ public:
 	FJsonValueArray(TArray< TSharedPtr<FJsonValue> >&& InArray) : Value(MoveTemp(InArray)) {Type = EJson::Array;}
 	virtual bool TryGetArray(const TArray< TSharedPtr<FJsonValue> >*& OutArray) const override	{ OutArray = &Value; return true; }
 	virtual bool TryGetArray(TArray< TSharedPtr<FJsonValue> >*& OutArray) override				{ OutArray = &Value; return true; }
-	
+	virtual SIZE_T GetMemoryFootprint() const override { return sizeof(*this) + GetAllocatedSize(); }
+
 protected:
 	TArray< TSharedPtr<FJsonValue> > Value;
 
 	virtual FString GetType() const override {return TEXT("Array");}
+	/** Helper to calculate allocated size of the Value array and its contents */
+	SIZE_T GetAllocatedSize() const 
+	{
+		SIZE_T SizeBytes = 0;
+		SizeBytes += Value.GetAllocatedSize();
+		for (const TSharedPtr<FJsonValue>& Element : Value)
+		{
+			SizeBytes += Element.IsValid() ? Element->GetMemoryFootprint() : 0;
+		}
+		return SizeBytes; 
+	}
 };
 
 
@@ -233,11 +260,13 @@ public:
 	FJsonValueObject(TSharedPtr<FJsonObject> InObject) : Value(MoveTemp(InObject)) {Type = EJson::Object;}
 	virtual bool TryGetObject(const TSharedPtr<FJsonObject>*& OutObject) const override	{ OutObject = &Value; return true; }
 	virtual bool TryGetObject(TSharedPtr<FJsonObject>*& OutObject) override				{ OutObject = &Value; return true; }
-	
+	virtual SIZE_T GetMemoryFootprint() const override { return sizeof(*this) + GetAllocatedSize(); }
 protected:
 	TSharedPtr<FJsonObject> Value;
 
 	virtual FString GetType() const override {return TEXT("Object");}
+	/** Helper to calculate allocated size of the Value object and its contents */
+	JSON_API SIZE_T GetAllocatedSize() const;
 };
 
 
@@ -246,6 +275,7 @@ class FJsonValueNull : public FJsonValue
 {
 public:
 	FJsonValueNull() {Type = EJson::Null;}
+	virtual SIZE_T GetMemoryFootprint() const override { return sizeof(*this); }
 
 protected:
 	virtual FString GetType() const override {return TEXT("Null");}
diff --git a/Engine/Source/Runtime/JsonUtilities/Private/JsonObjectWrapper.cpp b/Engine/Source/Runtime/JsonUtilities/Private/JsonObjectWrapper.cpp
--- a/Engine/Source/Runtime/JsonUtilities/Private/JsonObjectWrapper.cpp
+++ b/Engine/Source/Runtime/JsonUtilities/Private/JsonObjectWrapper.cpp
@@ -101,4 +101,15 @@ bool FJsonObjectWrapper::JsonObjectFromString(const FString& Str)
 	return FJsonSerializer::Deserialize(JsonReader, JsonObject);
 }
 
+SIZE_T FJsonObjectWrapper::GetAllocatedSize() const
+{
+	SIZE_T SizeBytes = 0;
+
+	SizeBytes += JsonString.GetAllocatedSize();
+	// NOTE - Given JsonObject is a shared ptr it's possible the underlying object is referenced by multiple things, 
+	// so there's potential memory here could be getting counted multiple times by multiple wrappers
+	SizeBytes += JsonObject.IsValid() ? JsonObject->GetMemoryFootprint() : 0;
+
+	return SizeBytes;
+}
 
diff --git a/Engine/Source/Runtime/JsonUtilities/Public/JsonObjectWrapper.h b/Engine/Source/Runtime/JsonUtilities/Public/JsonObjectWrapper.h
--- a/Engine/Source/Runtime/JsonUtilities/Public/JsonObjectWrapper.h
+++ b/Engine/Source/Runtime/JsonUtilities/Public/JsonObjectWrapper.h
@@ -39,6 +39,11 @@ public:
 
 	JSONUTILITIES_API bool JsonObjectToString(FString& Str) const;
 	JSONUTILITIES_API bool JsonObjectFromString(const FString& Str);
+
+	/**
+	 * Returns the memory allocated by this object in Bytes, should NOT include sizeof(*this).
+	 */
+	JSONUTILITIES_API SIZE_T GetAllocatedSize() const;
 };
 
 template<>
