[Json] JsonPointer now uses (expected) strings (not chars) for parsing delimeters

#rnx

diff --git a/Engine/Source/Runtime/Json/Private/JsonUtils/JsonPointer.cpp b/Engine/Source/Runtime/Json/Private/JsonUtils/JsonPointer.cpp
--- a/Engine/Source/Runtime/Json/Private/JsonUtils/JsonPointer.cpp
+++ b/Engine/Source/Runtime/Json/Private/JsonUtils/JsonPointer.cpp
@@ -8,9 +8,11 @@ namespace UE::Json
 {
 	namespace Private
 	{
-		static constexpr TCHAR RootCharacter = TEXT('#');
-		static constexpr TCHAR PathDelimiterCharacter = TEXT('/');
-		static constexpr TCHAR ArrayAppendCharacter = TEXT('-'); // unused, but kept for reference as official part of spec
+		// Stored as null-terminated strings in order to work with built-in string functions
+
+		static constexpr TCHAR RootCharacterStr[] = TEXT("#");
+		static constexpr TCHAR PathDelimiterCharacterStr[] = TEXT("/");
+		static constexpr TCHAR ArrayAppendCharacterStr[] = TEXT("-"); // unused, but kept for reference as official part of spec
 	}
 
 	FJsonPointer::FJsonPointer(FStringView InPath)
@@ -125,7 +127,7 @@ namespace UE::Json
 	bool FJsonPointer::ParsePath(const FString& InPath)
 	{
 		TArray<FString> PathStrings;
-		InPath.ParseIntoArray(PathStrings, &Private::PathDelimiterCharacter, false);
+		InPath.ParseIntoArray(PathStrings, Private::PathDelimiterCharacterStr, false);
 
 		if(InPath.IsEmpty())
 		{
@@ -133,7 +135,7 @@ namespace UE::Json
 		}
 
 		// We don't really care about the document root, remove it if it exists
-		if (PathStrings[0] == &Private::RootCharacter)
+		if (PathStrings[0] == Private::RootCharacterStr)
 		{
 			PathStrings.RemoveAt(0);
 		}
@@ -187,7 +189,7 @@ namespace UE::Json
 
 	FString FJsonPointer::ToString() const
 	{
-		return Private::RootCharacter + Private::PathDelimiterCharacter + *FString::Join(PathParts, &Private::PathDelimiterCharacter);
+		return FString(Private::RootCharacterStr) + Private::PathDelimiterCharacterStr + FString::Join(PathParts, Private::PathDelimiterCharacterStr);
 	}
 
 	FString FJsonPointer::EscapePart(const FString& Part)
