Support reading -nan or -nan(ind) in JsonReader(case insensitive).

#tests Tested through new added json test code through FoundationTests.

diff --git a/Engine/Source/Runtime/Json/Private/Tests/JsonTests.cpp b/Engine/Source/Runtime/Json/Private/Tests/JsonTests.cpp
--- a/Engine/Source/Runtime/Json/Private/Tests/JsonTests.cpp
+++ b/Engine/Source/Runtime/Json/Private/Tests/JsonTests.cpp
@@ -804,11 +804,16 @@ TEST_CASE_NAMED(FJsonAutomationTest, "System::Engine::FileSystem::JSON", "[Appli
 
 	// Test Nan
 	{
+		const FString TestNanInd = FString::Printf(TEXT("%.17g"), std::numeric_limits<double>::quiet_NaN());
+		CHECK(TestNanInd == TEXT("nan")); // Make sure code will not run on standard library impl which outputs nan(ind)
+
 		const FString InputString =
 			TEXT(
 				"{"
 					"\"Value0\":nan,"
-					"\"Value1\":NaN"
+					"\"Value1\":NaN,"
+					"\"Value2\":-nan,"
+					"\"Value3\":-nan(ind)"
 				"}"
 			);
 		TSharedRef< TJsonReader<> > Reader = TJsonReaderFactory<>::Create( InputString );
@@ -820,10 +825,16 @@ TEST_CASE_NAMED(FJsonAutomationTest, "System::Engine::FileSystem::JSON", "[Appli
 
 		const TSharedPtr<FJsonValue>* Value0 = Object->Values.Find(FString::Printf(TEXT("Value%i"), 0));
 		const TSharedPtr<FJsonValue>* Value1 = Object->Values.Find(FString::Printf(TEXT("Value%i"), 1));
+		const TSharedPtr<FJsonValue>* Value2 = Object->Values.Find(FString::Printf(TEXT("Value%i"), 2));
+		const TSharedPtr<FJsonValue>* Value3 = Object->Values.Find(FString::Printf(TEXT("Value%i"), 3));
 		const double Number0 = (*Value0)->AsNumber();
 		const double Number1 = (*Value1)->AsNumber();
+		const double Number2 = (*Value2)->AsNumber();
+		const double Number3 = (*Value3)->AsNumber();
 		REQUIRE(FMath::IsNaN(Number0));
 		REQUIRE(FMath::IsNaN(Number1));
+		REQUIRE(FMath::IsNaN(Number2));
+		REQUIRE(FMath::IsNaN(Number3));
 
 		FString OutputString;
 		TSharedRef< FCondensedJsonStringWriter > Writer = FCondensedJsonStringWriterFactory::Create( &OutputString );
@@ -834,10 +845,12 @@ TEST_CASE_NAMED(FJsonAutomationTest, "System::Engine::FileSystem::JSON", "[Appli
 			TEXT(
 				"{"
 					"\"Value0\":%.17g,"
-					"\"Value1\":%.17g"
+					"\"Value1\":%.17g,"
+					"\"Value2\":%.17g,"
+					"\"Value3\":%.17g"
 				"}"
 			),
-			std::numeric_limits<double>::quiet_NaN(), std::numeric_limits<double>::quiet_NaN());
+			std::numeric_limits<double>::quiet_NaN(), std::numeric_limits<double>::quiet_NaN(), -std::numeric_limits<double>::quiet_NaN(), -std::numeric_limits<double>::quiet_NaN());
 		REQUIRE(OutputString == TestOutput);
 	}
 
diff --git a/Engine/Source/Runtime/Json/Public/Serialization/JsonReader.h b/Engine/Source/Runtime/Json/Public/Serialization/JsonReader.h
--- a/Engine/Source/Runtime/Json/Public/Serialization/JsonReader.h
+++ b/Engine/Source/Runtime/Json/Public/Serialization/JsonReader.h
@@ -422,13 +422,17 @@ private:
 			{
 				if (IsJsonNumber(Char))
 				{
-					if (!ParseNumberToken(Char))
+					bool bParseNumberSucceed = ParseNumberToken(Char);
+					if (!bParseNumberSucceed && Char != '-') // Could be -NaN, will return false when fail to parse it as -NaN later on
 					{
 						return false;
 					}
 
-					OutToken = EJsonToken::Number;
-					return true;
+					if (bParseNumberSucceed)
+					{
+						OutToken = EJsonToken::Number;
+						return true;
+					}
 				}
 
 				switch (Char)
@@ -493,6 +497,7 @@ private:
 				case CharType('t'): case CharType('T'):
 				case CharType('f'): case CharType('F'):
 				case CharType('n'): case CharType('N'):
+				case CharType('-'):
 					{
 						FString Test;
 						Test += Char;
@@ -504,7 +509,8 @@ private:
 								return false;
 							}
 
-							if (IsAlphaNumber(Char))
+							if (IsAlphaNumber(Char) 
+								|| Char == '(' || Char == ')') // Could be "-nan(ind)" depending on the platform and impl of standard library when write
 							{
 								++CharacterNumber;
 								Test += Char;
@@ -544,6 +550,14 @@ private:
 							return true;
 						}
 
+						if (Test.Compare(TEXT("-NaN"), ESearchCase::IgnoreCase) == 0 ||
+							Test.Compare(TEXT("-NaN(ind)"), ESearchCase::IgnoreCase) == 0)
+						{
+							NumberValue = -std::numeric_limits<double>::quiet_NaN();
+							OutToken = EJsonToken::Number;
+							return true;
+						}
+
 						SetErrorMessage( TEXT("Invalid Json Token. Check that your member names have quotes around them!") );
 						return false;
 					}
@@ -795,7 +809,11 @@ private:
 			return true;
 		}
 
-		SetErrorMessage( TEXT("Poorly formed Json Number Token.") );
+		if (FirstChar != '-') // Could be -NaN, will set the error message when fail to parse it as -NaN later on
+		{
+			SetErrorMessage( TEXT("Poorly formed Json Number Token.") );
+		}
+
 		return false;
 	}
 
