Add support for creating json objects from utf8 strings, and utf8 strings from json objects

#prefligth 671fb450844fc1f144225579

diff --git a/Engine/Plugins/Enterprise/GLTFExporter/Source/GLTFExporter/Private/Json/GLTFJsonWriter.cpp b/Engine/Plugins/Enterprise/GLTFExporter/Source/GLTFExporter/Private/Json/GLTFJsonWriter.cpp
--- a/Engine/Plugins/Enterprise/GLTFExporter/Source/GLTFExporter/Private/Json/GLTFJsonWriter.cpp
+++ b/Engine/Plugins/Enterprise/GLTFExporter/Source/GLTFExporter/Private/Json/GLTFJsonWriter.cpp
@@ -158,9 +158,10 @@ struct TGLTFJsonPrintPolicyUTF8
 		Stream->Serialize(Buffer, Length);
 	}
 
-	static void WriteString(FArchive* Stream, const FString& String)
+	template <typename CharType>
+	static void WriteStringView(FArchive* Stream, TStringView<CharType> String)
 	{
-		const TCHAR* CharPtr = *String;
+		const CharType* CharPtr = String.GetData();
 
 		for (int32 CharIndex = 0; CharIndex < String.Len(); ++CharIndex, ++CharPtr)
 		{
@@ -168,24 +169,14 @@ struct TGLTFJsonPrintPolicyUTF8
 		}
 	}
 
-	static inline void WriteString(FArchive* Stream, FStringView String)
+	static void WriteString(FArchive* Stream, FStringView String)
 	{
-		const TCHAR* CharPtr = String.GetData();
-
-		for (int32 CharIndex = 0; CharIndex < String.Len(); ++CharIndex, ++CharPtr)
-		{
-			WriteChar(Stream, *CharPtr);
-		}
+		WriteStringView(Stream, String);
 	}
 
-	static inline void WriteString(FArchive* Stream, FAnsiStringView String)
+	static void WriteString(FArchive* Stream, FUtf8StringView String)
 	{
-		const ANSICHAR* CharPtr = String.GetData();
-
-		for (int32 CharIndex = 0; CharIndex < String.Len(); ++CharIndex, ++CharPtr)
-		{
-			WriteChar(Stream, *CharPtr);
-		}
+		WriteStringView(Stream, String);
 	}
 
 	static void WriteFloat(FArchive* Stream, float Number)
diff --git a/Engine/Plugins/VirtualProduction/LiveLinkMasterLockit/Source/LiveLinkMasterLockit/Private/MasterLockitMessageThread.cpp b/Engine/Plugins/VirtualProduction/LiveLinkMasterLockit/Source/LiveLinkMasterLockit/Private/MasterLockitMessageThread.cpp
--- a/Engine/Plugins/VirtualProduction/LiveLinkMasterLockit/Source/LiveLinkMasterLockit/Private/MasterLockitMessageThread.cpp
+++ b/Engine/Plugins/VirtualProduction/LiveLinkMasterLockit/Source/LiveLinkMasterLockit/Private/MasterLockitMessageThread.cpp
@@ -257,7 +257,7 @@ void FMasterLockitMessageThread::ParsePacket(FArrayReader& InPacketData)
 			TSharedPtr<FJsonObject> JsonObject = MakeShared<FJsonObject>();
 			FString JsonString = Payload.GetData();
 			JsonString.RemoveAt(Payload.Num(), JsonString.Len() + 1 - Payload.Num());
-			TSharedRef<TJsonReader<>> JsonReader = TJsonReaderFactory<ANSICHAR>::Create(JsonString);
+			TSharedRef<TJsonReader<>> JsonReader = TJsonReaderFactory<>::Create(JsonString);
 			FJsonSerializer::Deserialize(JsonReader, JsonObject);
 
 			// Acquire references to metadata categories that are present in all DeviceTypes 
diff --git a/Engine/Source/Runtime/Json/Private/Dom/JsonObject.cpp b/Engine/Source/Runtime/Json/Private/Dom/JsonObject.cpp
--- a/Engine/Source/Runtime/Json/Private/Dom/JsonObject.cpp
+++ b/Engine/Source/Runtime/Json/Private/Dom/JsonObject.cpp
@@ -97,6 +97,11 @@ FString FJsonObject::GetStringField(FStringView FieldName) const
 	return GetField<EJson::None>(FieldName)->AsString();
 }
 
+FUtf8String FJsonObject::GetUtf8StringField(FStringView FieldName) const
+{
+	return GetField<EJson::None>(FieldName)->AsUtf8String();
+}
+
 bool FJsonObject::TryGetStringField(FStringView FieldName, FString& OutString) const
 {
 	TSharedPtr<FJsonValue> Field = TryGetField(FieldName);
@@ -136,7 +141,7 @@ bool FJsonObject::TryGetStringArrayField(FStringView FieldName, TArray<FString>&
 
 void FJsonObject::SetStringField(FString&& FieldName, FString&& StringValue)
 {
-	this->Values.Add(MoveTemp(FieldName), MakeShared<FJsonValueString>(MoveTemp(StringValue)));
+	this->Values.Add(MoveTemp(FieldName), MakeShared<TJsonValueString<TCHAR>>(MoveTemp(StringValue)));
 }
 
 void FJsonObject::SetStringField(FString&& FieldName, const FString& StringValue)
@@ -154,6 +159,26 @@ void FJsonObject::SetStringField(const FString& FieldName, const FString& String
 	SetStringField(CopyTemp(FieldName), CopyTemp(StringValue));
 }
 
+void FJsonObject::SetStringField(FString&& FieldName, FUtf8String&& StringValue)
+{
+	this->Values.Add(MoveTemp(FieldName), MakeShared<TJsonValueString<UTF8CHAR>>(MoveTemp(StringValue)));
+}
+
+void FJsonObject::SetStringField(FString&& FieldName, const FUtf8String& StringValue)
+{
+	SetStringField(MoveTemp(FieldName), CopyTemp(StringValue));
+}
+
+void FJsonObject::SetStringField(const FString& FieldName, FUtf8String&& StringValue)
+{
+	SetStringField(CopyTemp(FieldName), MoveTemp(StringValue));
+}
+
+void FJsonObject::SetStringField(const FString& FieldName, const FUtf8String& StringValue)
+{
+	SetStringField(CopyTemp(FieldName), CopyTemp(StringValue));
+}
+
 bool FJsonObject::GetBoolField(FStringView FieldName) const
 {
 	return GetField<EJson::None>(FieldName)->AsBool();
diff --git a/Engine/Source/Runtime/Json/Private/Dom/JsonValue.cpp b/Engine/Source/Runtime/Json/Private/Dom/JsonValue.cpp
--- a/Engine/Source/Runtime/Json/Private/Dom/JsonValue.cpp
+++ b/Engine/Source/Runtime/Json/Private/Dom/JsonValue.cpp
@@ -28,6 +28,17 @@ FString FJsonValue::AsString() const
 	return String;
 }
 
+FUtf8String FJsonValue::AsUtf8String() const
+{
+	FUtf8String String;
+
+	if (!TryGetUtf8String(String))
+	{
+		ErrorMessage(TEXT("Utf8String"));
+	}
+
+	return String;
+}
 
 bool FJsonValue::AsBool() const 
 {
diff --git a/Engine/Source/Runtime/Json/Private/Serialization/JsonSerializable.cpp b/Engine/Source/Runtime/Json/Private/Serialization/JsonSerializable.cpp
--- a/Engine/Source/Runtime/Json/Private/Serialization/JsonSerializable.cpp
+++ b/Engine/Source/Runtime/Json/Private/Serialization/JsonSerializable.cpp
@@ -12,27 +12,59 @@ const FString FJsonSerializable::ToJson(bool bPrettyPrint/*=true*/) const
 {
 	// Strip away const, because we use a single method that can read/write which requires non-const semantics
 	// Otherwise, we'd have to have 2 separate macros for declaring const to json and non-const from json
-	return ((FJsonSerializable*)this)->ToJson(bPrettyPrint);
+	return const_cast<FJsonSerializable*>(this)->ToJson(bPrettyPrint);
 }
 
-const FString FJsonSerializable::ToJson(bool bPrettyPrint/*=true*/)
+const FUtf8String FJsonSerializable::ToJsonUtf8(bool bPrettyPrint/*=true*/) const
 {
-	FString JsonStr;
-	if (bPrettyPrint)
-	{
-		TSharedRef<TJsonWriter<> > JsonWriter = TJsonWriterFactory<>::Create(&JsonStr);
-		FJsonSerializerWriter<> Serializer(JsonWriter);
-		Serialize(Serializer, false);
-		JsonWriter->Close();
-	}
-	else
+	// Strip away const, because we use a single method that can read/write which requires non-const semantics
+	// Otherwise, we'd have to have 2 separate macros for declaring const to json and non-const from json
+	return const_cast<FJsonSerializable*>(this)->ToJsonUtf8(bPrettyPrint);
+}
+
+namespace FJsonSerializablePrivate
+{
+	template <typename StringType>
+	const StringType ToJsonStringWithType(FJsonSerializable& Json, bool bPrettyPrint)
 	{
-		TSharedRef< TJsonWriter< TCHAR, TCondensedJsonPrintPolicy< TCHAR > > > JsonWriter = TJsonWriterFactory< TCHAR, TCondensedJsonPrintPolicy< TCHAR > >::Create( &JsonStr );
-		FJsonSerializerWriter<TCHAR, TCondensedJsonPrintPolicy< TCHAR >> Serializer(JsonWriter);
-		Serialize(Serializer, false);
-		JsonWriter->Close();
+		using CharType = typename StringType::ElementType;
+		StringType JsonStr;
+		if (bPrettyPrint)
+		{
+			TSharedRef<TJsonWriter<CharType>> JsonWriter = TJsonWriterFactory<CharType>::Create(&JsonStr);
+			FJsonSerializerWriter<CharType> Serializer(JsonWriter);
+			Json.Serialize(Serializer, false);
+			JsonWriter->Close();
+		}
+		else
+		{
+			TSharedRef<TJsonWriter<CharType, TCondensedJsonPrintPolicy<CharType>>> JsonWriter = TJsonWriterFactory<CharType, TCondensedJsonPrintPolicy<CharType>>::Create(&JsonStr);
+			FJsonSerializerWriter<CharType, TCondensedJsonPrintPolicy<CharType>> Serializer(JsonWriter);
+			Json.Serialize(Serializer, false);
+			JsonWriter->Close();
+		}
+		return JsonStr;
 	}
-	return JsonStr;
+}
+
+const FString FJsonSerializable::ToJson(bool bPrettyPrint/*=true*/)
+{
+	return FJsonSerializablePrivate::ToJsonStringWithType<FString>(*this, bPrettyPrint);
+}
+
+const FUtf8String FJsonSerializable::ToJsonUtf8(bool bPrettyPrint/*=true*/)
+{
+	return FJsonSerializablePrivate::ToJsonStringWithType<FUtf8String>(*this, bPrettyPrint);
+}
+
+bool FJsonSerializable::FromJson(const TCHAR* Json)
+{
+	return FromJsonStringView(FStringView(Json));
+}
+
+bool FJsonSerializable::FromJson(const UTF8CHAR* Json)
+{
+	return FromJsonStringView(FUtf8StringView(Json));
 }
 
 bool FJsonSerializable::FromJson(const FString& Json)
@@ -40,11 +72,16 @@ bool FJsonSerializable::FromJson(const FString& Json)
 	return FromJsonStringView(FStringView(Json));
 }
 
+bool FJsonSerializable::FromJson(const FUtf8String& Json)
+{
+	return FromJsonStringView(FUtf8StringView(Json));
+}
+
 bool FJsonSerializable::FromJson(FString&& Json)
 {
 	TSharedPtr<FJsonObject> JsonObject;
 	TSharedRef<TJsonReader<> > JsonReader = TJsonReaderFactory<>::Create(MoveTemp(Json));
-	if (FJsonSerializer::Deserialize(JsonReader,JsonObject) &&
+	if (FJsonSerializer::Deserialize(JsonReader, JsonObject) &&
 		JsonObject.IsValid())
 	{
 		FJsonSerializerReader Serializer(JsonObject);
@@ -55,6 +92,11 @@ bool FJsonSerializable::FromJson(FString&& Json)
 	return false;
 }
 
+bool FJsonSerializable::FromJson(FUtf8String&& Json)
+{
+	return FromJsonStringView(FUtf8StringView(Json));
+}
+
 namespace UE::JsonSerializable::Private
 {
 
diff --git a/Engine/Source/Runtime/Json/Private/Serialization/JsonSerializerReader.cpp b/Engine/Source/Runtime/Json/Private/Serialization/JsonSerializerReader.cpp
--- a/Engine/Source/Runtime/Json/Private/Serialization/JsonSerializerReader.cpp
+++ b/Engine/Source/Runtime/Json/Private/Serialization/JsonSerializerReader.cpp
@@ -98,6 +98,14 @@ void FJsonSerializerReader::Serialize(FStringView Name, FString& Value)
 	}
 }
 
+void FJsonSerializerReader::Serialize(FStringView Name, FUtf8String& Value)
+{
+	if (JsonObject->HasTypedField<EJson::String>(Name))
+	{
+		Value = JsonObject->GetUtf8StringField(Name);
+	}
+}
+
 void FJsonSerializerReader::Serialize(FStringView Name, FText& Value) 
 {
 	if (JsonObject->HasTypedField<EJson::String>(Name))
@@ -318,3 +326,9 @@ void FJsonSerializerReader::WriteRawJSONValue(FStringView Value)
 	// Should never be called on a reader
 	check(false);
 }
+
+void FJsonSerializerReader::WriteRawJSONValue(FUtf8StringView Value)
+{
+	// Should never be called on a reader
+	check(false);
+}
diff --git a/Engine/Source/Runtime/Json/Private/Tests/JsonTests.cpp b/Engine/Source/Runtime/Json/Private/Tests/JsonTests.cpp
--- a/Engine/Source/Runtime/Json/Private/Tests/JsonTests.cpp
+++ b/Engine/Source/Runtime/Json/Private/Tests/JsonTests.cpp
@@ -3,6 +3,7 @@
 #if WITH_TESTS
 
 #include "CoreMinimal.h"
+#include "Containers/AnsiString.h"
 #include "Misc/AutomationTest.h"
 #include "Policies/CondensedJsonPrintPolicy.h"
 #include "Policies/PrettyJsonPrintPolicy.h"
@@ -17,6 +18,10 @@
 
 typedef TJsonWriterFactory< TCHAR, TCondensedJsonPrintPolicy<TCHAR> > FCondensedJsonStringWriterFactory;
 typedef TJsonWriter< TCHAR, TCondensedJsonPrintPolicy<TCHAR> > FCondensedJsonStringWriter;
+typedef TJsonWriterFactory< ANSICHAR, TCondensedJsonPrintPolicy<ANSICHAR> > FCondensedJsonAnsiStringWriterFactory;
+typedef TJsonWriter< ANSICHAR, TCondensedJsonPrintPolicy<ANSICHAR> > FCondensedJsonAnsiStringWriter;
+typedef TJsonWriterFactory< UTF8CHAR, TCondensedJsonPrintPolicy<UTF8CHAR> > FCondensedJsonUtf8StringWriterFactory;
+typedef TJsonWriter< UTF8CHAR, TCondensedJsonPrintPolicy<UTF8CHAR> > FCondensedJsonUtf8StringWriter;
 
 typedef TJsonWriterFactory< TCHAR, TPrettyJsonPrintPolicy<TCHAR> > FPrettyJsonStringWriterFactory;
 typedef TJsonWriter< TCHAR, TPrettyJsonPrintPolicy<TCHAR> > FPrettyJsonStringWriter;
@@ -695,6 +700,40 @@ TEST_CASE_NAMED(FJsonAutomationTest, "System::Engine::FileSystem::JSON", "[Appli
 		REQUIRE(OutputString == TestOutput);
 	}
 
+	// String Test ANSI
+	{
+		const ANSICHAR* InputString =
+			"{" \
+				"\"Value\":\"Some String, Escape Chars: \\\\, \\\", \\/, \\b, \\f, \\n, \\r, \\t, \\u002B\\uD83D\\uDE10\""\
+			"}";
+
+		TSharedRef< TJsonReader<ANSICHAR> > Reader = TJsonReaderFactory<ANSICHAR>::CreateFromView(InputString);
+
+		TSharedPtr<FJsonObject> Object;
+		bool bSuccessful = FJsonSerializer::Deserialize(Reader, Object);
+		REQUIRE(bSuccessful);
+		REQUIRE(Object.IsValid());
+
+		{
+			const TSharedPtr<FJsonValue>* Value = Object->Values.Find(TEXT("Value"));
+			REQUIRE((Value && (*Value)->Type == EJson::String));
+			const FString String = (*Value)->AsString();
+			REQUIRE(String == TEXT("Some String, Escape Chars: \\, \", /, \b, \f, \n, \r, \t, +😐"));
+			const FUtf8String Utf8String = (*Value)->AsUtf8String();
+			REQUIRE(Utf8String == UTF8TEXT("Some String, Escape Chars: \\, \", /, \b, \f, \n, \r, \t, +😐"));
+		}
+
+		FAnsiString AnsiOutputString;
+		TSharedRef< FCondensedJsonAnsiStringWriter > AnsiWriter = FCondensedJsonAnsiStringWriterFactory::Create(&AnsiOutputString);
+		REQUIRE(FJsonSerializer::Serialize(Object.ToSharedRef(), AnsiWriter));
+
+		const FAnsiString AnsiTestOutput =
+			"{" \
+				"\"Value\":\"Some String, Escape Chars: \\\\, \\\", /, \\b, \\f, \\n, \\r, \\t, +\\ud83d\\ude10\""\
+			"}";
+		REQUIRE(AnsiOutputString == AnsiTestOutput);
+	}
+
 	// String Test UTF8
 	{
 		// UTF8TEXT will prepend the first u8 literal specifier
@@ -720,24 +759,32 @@ TEST_CASE_NAMED(FJsonAutomationTest, "System::Engine::FileSystem::JSON", "[Appli
 			REQUIRE((Value && (*Value)->Type == EJson::String));
 			const FString String = (*Value)->AsString();
 			REQUIRE(String == TEXT("Some String, Escape Chars: \\, \", /, \b, \f, \n, \r, \t, +😐"));
+			const FUtf8String Utf8String = (*Value)->AsUtf8String();
+			REQUIRE(Utf8String == UTF8TEXT("Some String, Escape Chars: \\, \", /, \b, \f, \n, \r, \t, +😐"));
 		}
 		{
 			const TSharedPtr<FJsonValue>* Value = Object->Values.Find(TEXT("Value1"));
 			REQUIRE((Value && (*Value)->Type == EJson::String));
 			const FString String = (*Value)->AsString();
 			REQUIRE(String == TEXT("Greek String, Σὲ γνωρίζω ἀπὸ τὴν κόψη"));
+			const FUtf8String Utf8String = (*Value)->AsUtf8String();
+			REQUIRE(Utf8String == UTF8TEXT("Greek String, Σὲ γνωρίζω ἀπὸ τὴν κόψη"));
 		}
 		{
 			const TSharedPtr<FJsonValue>* Value = Object->Values.Find(TEXT("Value2"));
 			REQUIRE((Value && (*Value)->Type == EJson::String));
 			const FString String = (*Value)->AsString();
 			REQUIRE(String == TEXT("Thai String, สิบสองกษัตริย์ก่อนหน้าแลถัดไป"));
+			const FUtf8String Utf8String = (*Value)->AsUtf8String();
+			REQUIRE(Utf8String == UTF8TEXT("Thai String, สิบสองกษัตริย์ก่อนหน้าแลถัดไป"));
 		}
 		{
 			const TSharedPtr<FJsonValue>* Value = Object->Values.Find(TEXT("Value3"));
 			REQUIRE((Value && (*Value)->Type == EJson::String));
 			const FString String = (*Value)->AsString();
 			REQUIRE(String == TEXT("Hello world, Καλημέρα κόσμε, コンニチハ"));
+			const FUtf8String Utf8String = (*Value)->AsUtf8String();
+			REQUIRE(Utf8String == UTF8TEXT("Hello world, Καλημέρα κόσμε, コンニチハ"));
 		}
 
 		FString OutputString;
@@ -753,6 +800,20 @@ TEST_CASE_NAMED(FJsonAutomationTest, "System::Engine::FileSystem::JSON", "[Appli
 				TEXT("\"Value3\":\"Hello world, Καλημέρα κόσμε, コンニチハ\"")
 			TEXT("}");
 		REQUIRE(OutputString == TestOutput);
+
+		FUtf8String Utf8OutputString;
+		TSharedRef< FCondensedJsonUtf8StringWriter > Utf8Writer = FCondensedJsonUtf8StringWriterFactory::Create(&Utf8OutputString);
+		REQUIRE(FJsonSerializer::Serialize(Object.ToSharedRef(), Utf8Writer));
+
+		// Note: The literal prefix for the string (u8, L), must be present for every concatenated string, not just the first one
+		const FUtf8String Utf8TestOutput =
+			UTF8TEXT("{"\
+						"\"Value\":\"Some String, Escape Chars: \\\\, \\\", /, \\b, \\f, \\n, \\r, \\t, +😐\","\
+						"\"Value1\":\"Greek String, Σὲ γνωρίζω ἀπὸ τὴν κόψη\","\
+						"\"Value2\":\"Thai String, สิบสองกษัตริย์ก่อนหน้าแลถัดไป\","\
+						"\"Value3\":\"Hello world, Καλημέρα κόσμε, コンニチハ\""\
+					"}");
+		REQUIRE(Utf8OutputString == Utf8TestOutput);
 	}
 
 	// Number Test
diff --git a/Engine/Source/Runtime/Json/Public/Dom/JsonObject.h b/Engine/Source/Runtime/Json/Public/Dom/JsonObject.h
--- a/Engine/Source/Runtime/Json/Public/Dom/JsonObject.h
+++ b/Engine/Source/Runtime/Json/Public/Dom/JsonObject.h
@@ -184,6 +184,9 @@ public:
 	/** Get the field named FieldName as a string. */
 	JSON_API FString GetStringField(FStringView FieldName) const;
 
+	/** Get the field named FieldName as a UTF8 string. */
+	JSON_API FUtf8String GetUtf8StringField(FStringView FieldName) const;
+
 	/** Get the field named FieldName as a string. Returns false if it doesn't exist or cannot be converted. */
 	JSON_API bool TryGetStringField(FStringView FieldName, FString& OutString) const;
 
@@ -213,10 +216,39 @@ public:
 	}
 
 	/** Add a field named FieldName with value of StringValue */
+	void SetStringField(FString&& FieldName, const ANSICHAR* StringValue)
+	{
+		SetStringField(MoveTemp(FieldName), FUtf8String(StringValue));
+	}
+	void SetStringField(const FString& FieldName, const ANSICHAR* StringValue)
+	{
+		SetStringField(FieldName, FUtf8String(StringValue));
+	}
+	void SetStringField(FString&& FieldName, const TCHAR* StringValue)
+	{
+		SetStringField(MoveTemp(FieldName), FString(StringValue));
+	}
+	void SetStringField(const FString& FieldName, const TCHAR* StringValue)
+	{
+		SetStringField(FieldName, FString(StringValue));
+	}
+	void SetStringField(FString&& FieldName, const UTF8CHAR* StringValue)
+	{
+		SetStringField(MoveTemp(FieldName), FUtf8String(StringValue));
+	}
+	void SetStringField(const FString& FieldName, const UTF8CHAR* StringValue)
+	{
+		SetStringField(FieldName, FUtf8String(StringValue));
+	}
+
 	JSON_API void SetStringField(FString&& FieldName, FString&& StringValue);
 	JSON_API void SetStringField(FString&& FieldName, const FString& StringValue);
 	JSON_API void SetStringField(const FString& FieldName, FString&& StringValue);
 	JSON_API void SetStringField(const FString& FieldName, const FString& StringValue);
+	JSON_API void SetStringField(FString&& FieldName, FUtf8String&& StringValue);
+	JSON_API void SetStringField(FString&& FieldName, const FUtf8String& StringValue);
+	JSON_API void SetStringField(const FString& FieldName, FUtf8String&& StringValue);
+	JSON_API void SetStringField(const FString& FieldName, const FUtf8String& StringValue);
 
 	/**
 	 * Gets the field with the specified name as a boolean.
diff --git a/Engine/Source/Runtime/Json/Public/Dom/JsonValue.h b/Engine/Source/Runtime/Json/Public/Dom/JsonValue.h
--- a/Engine/Source/Runtime/Json/Public/Dom/JsonValue.h
+++ b/Engine/Source/Runtime/Json/Public/Dom/JsonValue.h
@@ -3,7 +3,9 @@
 #pragma once
 
 #include "Containers/Array.h"
+#include "Containers/ContainersFwd.h"
 #include "Containers/UnrealString.h"
+#include "Containers/Utf8String.h"
 #include "CoreMinimal.h"
 #include "HAL/Platform.h"
 #include "Misc/CString.h"
@@ -26,6 +28,9 @@ public:
 	/** Returns this value as a string, logging an error and returning an empty string if not possible */
 	JSON_API FString AsString() const;
 
+	/** Returns this value as a utf8 string, logging an error and returning an empty string if not possible */
+	JSON_API FUtf8String AsUtf8String() const;
+
 	/** Returns this value as a boolean, logging an error and returning false if not possible */
 	JSON_API bool AsBool() const;
 
@@ -68,6 +73,9 @@ public:
 	/** Tries to convert this value to a string, returning false if not possible */
 	virtual bool TryGetString(FString& OutString) const { return false; }
 
+	/** Tries to convert this value to a utf8 string, returning false if not possible */
+	virtual bool TryGetUtf8String(FUtf8String& OutString) const { return false; }
+
 	/** Tries to convert this value to a bool, returning false if not possible */
 	virtual bool TryGetBool(bool& OutBool) const { return false; }
 
@@ -129,21 +137,86 @@ protected:
 	}
 };
 
+template <typename CharType>
+struct TJsonValueStringType
+{
+	static_assert(sizeof(CharType) == 0, "Unsupported type");
+};
+
+template <>
+struct TJsonValueStringType<TCHAR>
+{
+	static FString GetType()
+	{
+		return TEXT("String");
+	}
+};
+
+template <>
+struct TJsonValueStringType<UTF8CHAR>
+{
+	static FString GetType()
+	{
+		return TEXT("Utf8String");
+	}
+};
 
 /** A Json String Value. */
-class FJsonValueString : public FJsonValue
+template <typename CharType>
+class TJsonValueString : public FJsonValue
 {
 public:
-	FJsonValueString(const FString& InString) : Value(InString) {Type = EJson::String;}
-	FJsonValueString(FString&& InString) : Value(MoveTemp(InString)) {Type = EJson::String;}
+	TJsonValueString(const TString<CharType>& InString) : Value(InString) {Type = EJson::String;}
+	TJsonValueString(TString<CharType>&& InString) : Value(MoveTemp(InString)) {Type = EJson::String;}
 
-	virtual bool TryGetString(FString& OutString) const override	{ OutString = Value; return true; }
-	virtual bool TryGetNumber(double& OutDouble) const override		{ if (Value.IsNumeric()) { OutDouble = FCString::Atod(*Value); return true; } else { return false; } }
-	virtual bool TryGetNumber(int32& OutValue) const override		{ LexFromString(OutValue, *Value); return true; }
-	virtual bool TryGetNumber(uint32& OutValue) const override		{ LexFromString(OutValue, *Value); return true; }
-	virtual bool TryGetNumber(int64& OutValue) const override		{ LexFromString(OutValue, *Value); return true; }
-	virtual bool TryGetNumber(uint64& OutValue) const override		{ LexFromString(OutValue, *Value); return true; }
-	virtual bool TryGetBool(bool& OutBool) const override			{ OutBool = Value.ToBool(); return true; }
+	virtual bool TryGetString(FString& OutString) const override
+	{
+		if constexpr (std::is_same_v<FString, TString<CharType>>)
+		{
+			OutString = Value;
+		}
+		else
+		{
+			OutString = FString(Value);
+		}
+		return true; 
+	}
+	
+	virtual bool TryGetUtf8String(FUtf8String& OutString) const override
+	{
+		if constexpr (std::is_same_v<FUtf8String, TString<CharType>>)
+		{
+			OutString = Value;
+		}
+		else
+		{
+			OutString = FUtf8String(Value);
+		}
+		return true;
+	}
+	
+	virtual bool TryGetNumber(double& OutDouble) const override
+	{
+		if (Value.IsNumeric())
+		{
+			OutDouble = TCString<CharType>::Atod(*Value);
+			return true;
+		}
+		else
+		{
+			return false;
+		}
+	}
+
+	virtual bool TryGetNumber(int32& OutValue) const override	{ LexFromString(OutValue, *Value); return true; }
+	
+	virtual bool TryGetNumber(uint32& OutValue) const override	{ LexFromString(OutValue, *Value); return true; }
+
+	virtual bool TryGetNumber(int64& OutValue) const override	{ LexFromString(OutValue, *Value); return true; }
+
+	virtual bool TryGetNumber(uint64& OutValue) const override	{ LexFromString(OutValue, *Value); return true; }
+
+	virtual bool TryGetBool(bool& OutBool) const override		{ OutBool = Value.ToBool(); return true; }
 
 	// Way to check if string value is empty without copying the string 
 	bool IsEmpty() const { return Value.IsEmpty(); }
@@ -151,13 +224,15 @@ public:
 	virtual SIZE_T GetMemoryFootprint() const override { return sizeof(*this) + GetAllocatedSize(); }
 
 protected:
-	FString Value;
+	TString<CharType> Value;
+
+	virtual FString GetType() const override { return TJsonValueStringType<CharType>::GetType(); }
 
-	virtual FString GetType() const override {return TEXT("String");}
 	/** Helper to calculate allocated size of the Value string */
 	SIZE_T GetAllocatedSize() const { return Value.GetAllocatedSize(); }
 };
 
+using FJsonValueString = TJsonValueString<TCHAR>;
 
 /** A Json Number Value. */
 class FJsonValueNumber : public FJsonValue
@@ -167,6 +242,7 @@ public:
 	virtual bool TryGetNumber(double& OutNumber) const override		{ OutNumber = Value; return true; }
 	virtual bool TryGetBool(bool& OutBool) const override			{ OutBool = (Value != 0.0); return true; }
 	virtual bool TryGetString(FString& OutString) const override	{ OutString = FString::SanitizeFloat(Value, 0); return true; }
+	virtual bool TryGetUtf8String(FUtf8String& OutString) const override { OutString = FUtf8String::SanitizeFloat(Value, 0); return true; }
 	virtual SIZE_T GetMemoryFootprint() const override { return sizeof(*this); }
 
 protected:
@@ -176,15 +252,65 @@ protected:
 	virtual FString GetType() const override {return TEXT("Number");}
 };
 
+template <typename CharType>
+struct TJsonValueNumberType
+{
+	static_assert(sizeof(CharType), "Unsupported type");
+};
+
+template <>
+struct TJsonValueNumberType<TCHAR>
+{
+	static FString GetType()
+	{
+		return TEXT("NumberString");
+	}
+};
+
+template <>
+struct TJsonValueNumberType<UTF8CHAR>
+{
+	static FString GetType()
+	{
+		return TEXT("Utf8NumberString");
+	}
+};
 
 /** A Json Number Value, stored internally as a string so as not to lose precision */
-class FJsonValueNumberString : public FJsonValue
+template <typename CharType>
+class TJsonValueNumberString : public FJsonValue
 {
 public:
-	FJsonValueNumberString(const FString& InString) : Value(InString) { Type = EJson::Number; }
-	FJsonValueNumberString(FString&& InString) : Value(MoveTemp(InString)) { Type = EJson::Number; }
+	TJsonValueNumberString(const TString<CharType>& InString) : Value(InString) { Type = EJson::Number; }
+
+	TJsonValueNumberString(TString<CharType>&& InString) : Value(MoveTemp(InString)) { Type = EJson::Number; }
 
-	virtual bool TryGetString(FString& OutString) const override { OutString = Value; return true; }
+	virtual bool TryGetString(FString& OutString) const override
+	{
+		if constexpr (std::is_same_v<FString, TString<CharType>>)
+		{
+			OutString = Value;
+		}
+		else
+		{
+			OutString = FString(Value);
+		}
+		return true;
+	}
+	
+	virtual bool TryGetUtf8String(FUtf8String& OutString) const override
+	{
+		if constexpr (std::is_same_v<FUtf8String, TString<CharType>>)
+		{
+			OutString = Value;
+		}
+		else
+		{
+			OutString = FUtf8String(Value);
+		}
+		return true;
+	}
+	
 	virtual bool TryGetNumber(double& OutDouble) const override { return LexTryParseString(OutDouble, *Value); }
 	virtual bool TryGetNumber(float &OutDouble) const override { return LexTryParseString(OutDouble, *Value); }
 	virtual bool TryGetNumber(int8& OutValue) const override { return LexTryParseString(OutValue, *Value); }
@@ -200,13 +326,14 @@ public:
 	virtual SIZE_T GetMemoryFootprint() const override { return sizeof(*this) + GetAllocatedSize(); }
 
 protected:
-	FString Value;
+	TString<CharType> Value;
 
-	virtual FString GetType() const override { return TEXT("NumberString"); }
+	virtual FString GetType() const override { return TJsonValueNumberType<CharType>::GetType(); }
 	/** Helper to calculate allocated size of the Value string */
 	SIZE_T GetAllocatedSize() const { return Value.GetAllocatedSize(); }
 };
 
+using FJsonValueNumberString = TJsonValueNumberString<TCHAR>;
 
 /** A Json Boolean Value. */
 class FJsonValueBoolean : public FJsonValue
@@ -216,6 +343,7 @@ public:
 	virtual bool TryGetNumber(double& OutNumber) const override		{ OutNumber = Value ? 1 : 0; return true; }
 	virtual bool TryGetBool(bool& OutBool) const override			{ OutBool = Value; return true; }
 	virtual bool TryGetString(FString& OutString) const override	{ OutString = Value ? TEXT("true") : TEXT("false"); return true; }
+	virtual bool TryGetUtf8String(FUtf8String& OutString) const override { OutString = Value ? UTF8TEXT("true") : UTF8TEXT("false"); return true; }
 	virtual SIZE_T GetMemoryFootprint() const override { return sizeof(*this); }
 
 protected:
diff --git a/Engine/Source/Runtime/Json/Public/Policies/JsonPrintPolicy.h b/Engine/Source/Runtime/Json/Public/Policies/JsonPrintPolicy.h
--- a/Engine/Source/Runtime/Json/Public/Policies/JsonPrintPolicy.h
+++ b/Engine/Source/Runtime/Json/Public/Policies/JsonPrintPolicy.h
@@ -2,6 +2,7 @@
 
 #pragma once
 
+#include "Containers/Utf8String.h"
 #include "CoreMinimal.h"
 
 /**
@@ -9,9 +10,11 @@
  *
  * @param CharType The type of characters to print, i.e. TCHAR or ANSICHAR.
  */
-template <class CharType>
+template <class InCharType>
 struct TJsonPrintPolicy
 {
+	using CharType = InCharType;
+
 	/**
 	 * Writes a single character to the output stream.
 	 *
@@ -23,18 +26,6 @@ struct TJsonPrintPolicy
 		Stream->Serialize(&Char, sizeof(CharType));
 	}
 
-	/**
-	 * Writes a string to the output stream.
-	 *
-	 * @param Stream The stream to write to.
-	 * @param String The string to write.
-	 */
-	static inline void WriteString( FArchive* Stream, const FString& String )
-	{
-		auto Conv = StringCast<CharType>(*String, String.Len());
-		Stream->Serialize((void*)Conv.Get(), Conv.Length() * sizeof(CharType));
-	}
-
 	/**
 	 * Writes a string to the output stream.
 	 *
@@ -48,12 +39,12 @@ struct TJsonPrintPolicy
 	}
 
 	/**
-	 * Writes a string to the output stream.
+	 * Writes a utf8 string to the output stream.
 	 *
 	 * @param Stream The stream to write to.
 	 * @param String The string to write.
 	 */
-	static inline void WriteString(FArchive* Stream, FAnsiStringView String)
+	static inline void WriteString(FArchive* Stream, FUtf8StringView String)
 	{
 		auto Conv = StringCast<CharType>(String.GetData(), String.Len());
 		Stream->Serialize((void*)Conv.Get(), Conv.Length() * sizeof(CharType));
@@ -87,18 +78,6 @@ struct TJsonPrintPolicy
 
 #if !PLATFORM_TCHAR_IS_CHAR16
 
-/**
- * Specialization for UTF16CHAR that writes FString data UTF-16.
- */
-template <>
-inline void TJsonPrintPolicy<UTF16CHAR>::WriteString(FArchive* Stream, const FString& String)
-{
-	// Note: This is a no-op on platforms that are using a 16-bit TCHAR
-	FTCHARToUTF16 UTF16String(*String, String.Len());
-
-	Stream->Serialize((void*)UTF16String.Get(), UTF16String.Length() * sizeof(UTF16CHAR));
-}
-
 /**
  * Specialization for UTF16CHAR that writes FString data UTF-16.
  */
diff --git a/Engine/Source/Runtime/Json/Public/Serialization/JsonReader.h b/Engine/Source/Runtime/Json/Public/Serialization/JsonReader.h
--- a/Engine/Source/Runtime/Json/Public/Serialization/JsonReader.h
+++ b/Engine/Source/Runtime/Json/Public/Serialization/JsonReader.h
@@ -3,6 +3,7 @@
 #pragma once
 
 #include "CoreMinimal.h"
+#include "Containers/Utf8String.h"
 #include "Serialization/JsonTypes.h"
 #include "Serialization/BufferReader.h"
 #include "Misc/StringBuilder.h"
@@ -38,6 +39,10 @@ template <class CharType = TCHAR>
 class TJsonReader
 {
 public:
+	// Store ANSICHAR in FUtf8String because JSON may contain scaped unicode characters that would not be recoverable otherwise. Use default string type for the char type otherwise
+	using StoredStringType = std::conditional_t<std::is_same_v<CharType, ANSICHAR>, FUtf8String, TString<CharType>>;
+	using StoredCharType = TElementType_T<StoredStringType>;
+	using ValueAsStringReturnType = std::conditional_t<std::is_same_v<StoredStringType, FString>, const FString&, FString>;
 
 	static TSharedRef< TJsonReader<CharType> > Create( FArchive* const Stream )
 	{
@@ -152,25 +157,45 @@ public:
 
 	FORCEINLINE virtual const FString& GetIdentifier() const { return Identifier; }
 
-	FORCEINLINE virtual  const FString& GetValueAsString() const 
-	{ 
+	FORCEINLINE virtual ValueAsStringReturnType GetValueAsString() const
+	{
+		check(CurrentToken == EJsonToken::String);
+		if constexpr (std::is_same_v<StoredStringType, FString>)
+		{
+			return StringValue;
+		}
+		else
+		{
+			// Construct FString from a different string type
+			return FString(StringValue);
+		}
+	}
+
+	FORCEINLINE virtual const StoredStringType& GetInternalValueAsString() const
+	{
 		check(CurrentToken == EJsonToken::String);
 		return StringValue;
 	}
-	
+
+	FORCEINLINE virtual StoredStringType StealInternalValueAsString()
+	{
+		check(CurrentToken == EJsonToken::String);
+		return MoveTemp(StringValue);
+	}
+
 	FORCEINLINE double GetValueAsNumber() const 
 	{ 
 		check(CurrentToken == EJsonToken::Number);
 		return NumberValue;
 	}
 
-	FORCEINLINE const FString& GetValueAsNumberString() const
+	FORCEINLINE const StoredStringType& GetValueAsNumberString() const
 	{
 		check(CurrentToken == EJsonToken::Number);
 		return StringValue;
 	}
 	
-	FORCEINLINE bool GetValueAsBoolean() const 
+	FORCEINLINE bool GetValueAsBoolean() const
 	{ 
 		check((CurrentToken == EJsonToken::True) || (CurrentToken == EJsonToken::False));
 		return BoolValue; 
@@ -334,7 +359,8 @@ private:
 				return false;
 			}
 
-			Identifier = StringValue;
+			// Move value if possible. A conversion will happen if StringValue is not an FString
+			Identifier = FString(MoveTemp(StringValue));
 			Token = EJsonToken::None;
 
 			if (!NextToken(Token))
@@ -499,7 +525,7 @@ private:
 				case CharType('n'): case CharType('N'):
 				case CharType('-'):
 					{
-						FString Test;
+						StoredStringType Test;
 						Test += Char;
 
 						while (!Stream->AtEnd())
@@ -523,27 +549,27 @@ private:
 							}
 						}
 
-						if (Test == TEXT("False"))
+						if (Test == CHARTEXT(CharType, "False"))
 						{
 							BoolValue = false;
 							OutToken = EJsonToken::False;
 							return true;
 						}
 
-						if (Test == TEXT("True"))
+						if (Test == CHARTEXT(CharType, "True"))
 						{
 							BoolValue = true;
 							OutToken = EJsonToken::True;
 							return true;
 						}
 
-						if (Test == TEXT("Null"))
+						if (Test == CHARTEXT(CharType, "Null"))
 						{
 							OutToken = EJsonToken::Null;
 							return true;
 						}
 
-						if (Test.Compare(TEXT("NaN"), ESearchCase::IgnoreCase) == 0)
+						if (Test.Compare(CHARTEXT(CharType, "NaN"), ESearchCase::IgnoreCase) == 0)
 						{
 							NumberValue = std::numeric_limits<double>::quiet_NaN();
 							OutToken = EJsonToken::Number;
@@ -575,7 +601,7 @@ private:
 
 	bool ParseStringToken()
 	{
-		TStringBuilderWithBuffer<CharType, 512> StringBuffer;
+		TStringBuilderWithBuffer<StoredCharType, 512> StringBuffer;
 		TStringBuilderWithBuffer<WIDECHAR, 512> UTF16CodePoints;
 
 		// Add escaped surrogate pairs
@@ -583,7 +609,7 @@ private:
 		{
 			if (UTF16CodePoints.Len() > 0)
 			{
-				// Will convert to CharType encoding if needed
+				// Will convert to StoredCharType encoding if needed
 				StringBuffer.Append(UTF16CodePoints);
 				UTF16CodePoints.Reset();
 			}
@@ -678,17 +704,17 @@ private:
 			}
 		}
 
-		StringValue = FString(StringBuffer);
+		StringValue = StoredStringType(StringBuffer);
 
 		// Inline combine any surrogate pairs in the data when loading into a UTF-32 string
-		StringConv::InlineCombineSurrogates(StringValue);
+		InlineCombineSurrogates(StringValue);
 
 		return true;
 	}
 
 	bool ParseNumberToken( CharType FirstChar )
 	{
-		FString String;
+		StoredStringType String;
 		int32 State = 0;
 		bool UseFirstChar = true;
 		bool StateError = false;
@@ -717,7 +743,7 @@ private:
 			}
 
 			// The following code doesn't actually derive the Json Number: that is handled
-			// by the function FCString::Atof below. This code only ensures the Json Number is
+			// by the function FPlatformString::Atod below. This code only ensures the Json Number is
 			// EXACTLY to specification
 			if (IsJsonNumber(Char))
 			{
@@ -805,7 +831,7 @@ private:
 		if (!StateError && ((State == 2) || (State == 3) || (State == 6) || (State == 8)))
 		{
 			StringValue = String;
-			NumberValue = FCString::Atod(*String);
+			NumberValue = FPlatformString::Atod(*String);
 			return true;
 		}
 
@@ -892,6 +918,14 @@ protected:
 		return true;
 	}
 
+	template <typename Type>
+	void InlineCombineSurrogates(Type& String)
+	{
+		if constexpr (std::is_same_v<Type, FString>)
+		{
+			StringConv::InlineCombineSurrogates(String);
+		}
+	}
 protected:
 
 	TArray<EJson> ParseState;
@@ -900,7 +934,7 @@ protected:
 	FArchive* Stream;
 	FString Identifier;
 	FString ErrorMessage;
-	FString StringValue;
+	StoredStringType StringValue;
 	double NumberValue;
 	uint32 LineNumber;
 	uint32 CharacterNumber;
@@ -909,28 +943,28 @@ protected:
 };
 
 
-class FJsonStringReader
-	: public TJsonReader<TCHAR>
+template <typename CharType>
+class TJsonStringReader
+	: public TJsonReader<CharType>
 {
 public:
-
-	static TSharedRef<FJsonStringReader> Create(const FString& JsonString)
+	static TSharedRef<TJsonStringReader> Create(const TString<CharType>& JsonString)
 	{
-		return MakeShareable(new FJsonStringReader(JsonString));
+		return MakeShareable(new TJsonStringReader(JsonString));
 	}
 
-	static TSharedRef<FJsonStringReader> Create(FString&& JsonString)
+	static TSharedRef<TJsonStringReader> Create(TString<CharType>&& JsonString)
 	{
-		return MakeShareable(new FJsonStringReader(MoveTemp(JsonString)));
+		return MakeShareable(new TJsonStringReader(MoveTemp(JsonString)));
 	}
 
-	const FString& GetSourceString() const
+	const TString<CharType>& GetSourceString() const
 	{
 		return Content;
 	}
 public:
 
-	virtual ~FJsonStringReader() = default;
+	virtual ~TJsonStringReader() = default;
 
 protected:
 
@@ -939,7 +973,7 @@ protected:
 	 *
 	 * @param JsonString The Json string to parse.
 	 */
-	explicit FJsonStringReader(const FString& JsonString)
+	explicit TJsonStringReader(const TString<CharType>& JsonString)
 		: Content(JsonString)
 		, Reader(nullptr)
 	{
@@ -951,7 +985,7 @@ protected:
 	 *
 	 * @param JsonString The Json string to parse.
 	 */
-	explicit FJsonStringReader(FString&& JsonString)
+	explicit TJsonStringReader(TString<CharType>&& JsonString)
 		: Content(MoveTemp(JsonString))
 		, Reader(nullptr)
 	{
@@ -965,15 +999,17 @@ protected:
 			return;
 		}
 
-		Reader = MakeUnique<FBufferReader>((void*)*Content, Content.Len() * sizeof(TCHAR), false);
-		Stream = Reader.Get();
+		Reader = MakeUnique<FBufferReader>((void*)*Content, Content.Len() * sizeof(CharType), false);
+		this->Stream = Reader.Get();
 	}
 
 protected:
-	const FString Content;
+	const TString<CharType> Content;
 	TUniquePtr<FBufferReader> Reader;
 };
 
+using FJsonStringReader = TJsonStringReader<TCHAR>;
+
 template <class CharType>
 class TJsonStringViewReader
 	: public TJsonReader<CharType>
@@ -1023,15 +1059,16 @@ template <class CharType = TCHAR>
 class TJsonReaderFactory
 {
 public:
+	using StringType = TString<CharType>;
 
-	static TSharedRef<TJsonReader<TCHAR>> Create(const FString& JsonString)
+	static TSharedRef<TJsonReader<TElementType_T<StringType>>> Create(StringType&& JsonString)
 	{
-		return FJsonStringReader::Create(JsonString);
+		return TJsonStringReader<TElementType_T<StringType>>::Create(MoveTemp(JsonString));
 	}
 
-	static TSharedRef<TJsonReader<TCHAR>> Create(FString&& JsonString)
+	static TSharedRef<TJsonReader<TElementType_T<StringType>>> Create(const StringType& JsonString)
 	{
-		return FJsonStringReader::Create(MoveTemp(JsonString));
+		return TJsonStringReader<TElementType_T<StringType>>::Create(JsonString);
 	}
 
 	static TSharedRef<TJsonReader<CharType>> Create(FArchive* const Stream)
diff --git a/Engine/Source/Runtime/Json/Public/Serialization/JsonSerializable.h b/Engine/Source/Runtime/Json/Public/Serialization/JsonSerializable.h
--- a/Engine/Source/Runtime/Json/Public/Serialization/JsonSerializable.h
+++ b/Engine/Source/Runtime/Json/Public/Serialization/JsonSerializable.h
@@ -26,7 +26,8 @@ struct FJsonSerializable
 	 * @return the corresponding json string
 	 */
 	JSON_API const FString ToJson(bool bPrettyPrint = true) const;
-
+	JSON_API const FUtf8String ToJsonUtf8(bool bPrettyPrint = true) const;
+	
 	/**
 	 * Serializes this object to its JSON string form
 	 *
@@ -34,8 +35,9 @@ struct FJsonSerializable
 	 * @return the corresponding json string
 	 */
 	JSON_API virtual const FString ToJson(bool bPrettyPrint=true);
-
-	/**
+	JSON_API virtual const FUtf8String ToJsonUtf8(bool bPrettyPrint = true);
+	
+		/**
 	 * Serializes this object with a Json Writer
 	 * 
 	 * @param JsonWriter - The writer to use
@@ -49,7 +51,10 @@ struct FJsonSerializable
 	 *
 	 * @param Json the JSON data to serialize from
 	 */
+	JSON_API virtual bool FromJson(const TCHAR* Json);
+	JSON_API virtual bool FromJson(const UTF8CHAR* Json);
 	JSON_API virtual bool FromJson(const FString& Json);
+	JSON_API virtual bool FromJson(const FUtf8String& Json);
 
 	/**
 	 * Serializes the contents of a JSON string into this object
@@ -57,6 +62,7 @@ struct FJsonSerializable
 	 * @param Json the JSON data to serialize from
 	 */
 	JSON_API virtual bool FromJson(FString&& Json);
+	JSON_API virtual bool FromJson(FUtf8String&& Json);
 
 	/**
 	 * Serializes the contents of a JSON string into this object using FUtf8StringView
@@ -263,4 +269,4 @@ static void ToJson(TArray<T>& InArray, Private::CondensedWriter& JsonWriter)
 	Private::ToJson(InArray, Private::ToJsonVariantArgs(TInPlaceType<Private::WriterVariants>(), Private::WriterVariants(TInPlaceType<Private::CondensedWriter>(), JsonWriter)));
 }
 
-} // namespace UE::JsonArray
+} // namespace UE::JsonArray
\ No newline at end of file
diff --git a/Engine/Source/Runtime/Json/Public/Serialization/JsonSerializer.h b/Engine/Source/Runtime/Json/Public/Serialization/JsonSerializer.h
--- a/Engine/Source/Runtime/Json/Public/Serialization/JsonSerializer.h
+++ b/Engine/Source/Runtime/Json/Public/Serialization/JsonSerializer.h
@@ -351,13 +351,17 @@ private:
 				break;
 
 			case EJsonNotation::String:
-				NewValue = MakeShared<FJsonValueString>(Reader.GetValueAsString());
+				{
+					using StoredCharType = typename TJsonReader<CharType>::StoredCharType;
+					NewValue = MakeShared<TJsonValueString<StoredCharType>>(Reader.StealInternalValueAsString());
+				}
 				break;
 
 			case EJsonNotation::Number:
 				if (EnumHasAnyFlags(InOptions, EFlags::StoreNumbersAsStrings))
 				{
-					NewValue = MakeShared<FJsonValueNumberString>(Reader.GetValueAsNumberString());
+					using StoredCharType = typename TJsonReader<CharType>::StoredCharType;
+					NewValue = MakeShared<TJsonValueNumberString<StoredCharType>>(Reader.GetValueAsNumberString());
 				}
 				else
 				{
diff --git a/Engine/Source/Runtime/Json/Public/Serialization/JsonSerializerBase.h b/Engine/Source/Runtime/Json/Public/Serialization/JsonSerializerBase.h
--- a/Engine/Source/Runtime/Json/Public/Serialization/JsonSerializerBase.h
+++ b/Engine/Source/Runtime/Json/Public/Serialization/JsonSerializerBase.h
@@ -23,6 +23,7 @@ struct FJsonSerializerBase
 	virtual void Serialize(FStringView Name, int64& Value) = 0;
 	virtual void Serialize(FStringView Name, bool& Value) = 0;
 	virtual void Serialize(FStringView Name, FString& Value) = 0;
+	virtual void Serialize(FStringView Name, FUtf8String& Value) = 0;
 	virtual void Serialize(FStringView Name, FText& Value) = 0;
 	virtual void Serialize(FStringView Name, float& Value) = 0;
 	virtual void Serialize(FStringView Name, double& Value) = 0;
@@ -43,6 +44,7 @@ struct FJsonSerializerBase
 	virtual TSharedPtr<FJsonObject> GetObject() = 0;
 	virtual void WriteIdentifierPrefix(FStringView Name) = 0;
 	virtual void WriteRawJSONValue(FStringView Value) = 0;
+	virtual void WriteRawJSONValue(FUtf8StringView Value) = 0;
 
 #if !PLATFORM_TCHAR_IS_UTF8CHAR
 
diff --git a/Engine/Source/Runtime/Json/Public/Serialization/JsonSerializerMacros.h b/Engine/Source/Runtime/Json/Public/Serialization/JsonSerializerMacros.h
--- a/Engine/Source/Runtime/Json/Public/Serialization/JsonSerializerMacros.h
+++ b/Engine/Source/Runtime/Json/Public/Serialization/JsonSerializerMacros.h
@@ -116,18 +116,19 @@
 #define JSON_SERIALIZE_RAW_JSON_STRING(JsonName, JsonValue) \
 		if (Serializer.IsLoading()) \
 		{ \
+			using CharType = TElementType_T<decltype(JsonValue)>; \
 			if (Serializer.GetObject()->HasTypedField<EJson::Object>(TEXTVIEW(JsonName))) \
 			{ \
 				TSharedPtr<FJsonObject> JsonObject = Serializer.GetObject()->GetObjectField(TEXTVIEW(JsonName)); \
 				if (JsonObject.IsValid()) \
 				{ \
-					auto Writer = TJsonWriterFactory<TCHAR, TCondensedJsonPrintPolicy<TCHAR>>::Create(&JsonValue); \
+					auto Writer = TJsonWriterFactory<CharType, TCondensedJsonPrintPolicy<CharType>>::Create(&JsonValue); \
 					FJsonSerializer::Serialize(JsonObject.ToSharedRef(), Writer); \
 				} \
 			} \
 			else \
 			{ \
-				JsonValue = FString(); \
+				JsonValue = TString<CharType>(); \
 			} \
 		} \
 		else \
diff --git a/Engine/Source/Runtime/Json/Public/Serialization/JsonSerializerReader.h b/Engine/Source/Runtime/Json/Public/Serialization/JsonSerializerReader.h
--- a/Engine/Source/Runtime/Json/Public/Serialization/JsonSerializerReader.h
+++ b/Engine/Source/Runtime/Json/Public/Serialization/JsonSerializerReader.h
@@ -74,6 +74,13 @@ public:
 	 * @param Value the out value to read the data into
 	 */
 	JSON_API virtual void Serialize(FStringView Name, FString& Value) override;
+	/**
+	 * If the underlying json object has the field, it is read into the value
+	 *
+	 * @param Name the name of the field to read
+	 * @param Value the out value to read the data into
+	 */
+	JSON_API virtual void Serialize(FStringView Name, FUtf8String& Value) override;
 	/**
 	 * If the underlying json object has the field, it is read into the value
 	 *
@@ -196,10 +203,11 @@ public:
 	 */
 	JSON_API virtual void SerializeMapSafe(FStringView Name, FJsonSerializableKeyValueMap& Map) override;
 
-	JSON_API virtual void WriteIdentifierPrefix(FStringView Name);
+	JSON_API virtual void WriteIdentifierPrefix(FStringView Name) override;
 
-	JSON_API virtual void WriteRawJSONValue(FStringView Value);
+	JSON_API virtual void WriteRawJSONValue(FStringView Value) override;
 
+	JSON_API virtual void WriteRawJSONValue(FUtf8StringView Value) override;
 private:
 	/** The object that holds the parsed JSON data */
 	TSharedPtr<FJsonObject> JsonObject;
diff --git a/Engine/Source/Runtime/Json/Public/Serialization/JsonSerializerWriter.h b/Engine/Source/Runtime/Json/Public/Serialization/JsonSerializerWriter.h
--- a/Engine/Source/Runtime/Json/Public/Serialization/JsonSerializerWriter.h
+++ b/Engine/Source/Runtime/Json/Public/Serialization/JsonSerializerWriter.h
@@ -127,6 +127,16 @@ public:
 	{
 		JsonWriter->WriteValue(Name, Value);
 	}
+	/**
+	* Writes the field name and the corresponding value to the JSON data
+	*
+	* @param Name the field name to write out
+	* @param Value the value to write out
+	*/
+	virtual void Serialize(FStringView Name, FUtf8String& Value) override
+	{
+		JsonWriter->WriteValue(Name, Value);
+	}
 	/**
 	 * Writes the field name and the corresponding value to the JSON data
 	 *
@@ -391,5 +401,10 @@ public:
 	{
 		JsonWriter->WriteRawJSONValue(Value);
 	}
+
+	virtual void WriteRawJSONValue(FUtf8StringView Value)
+	{
+		JsonWriter->WriteRawJSONValue(Value);
+	}
 };
 
diff --git a/Engine/Source/Runtime/Json/Public/Serialization/JsonWriter.h b/Engine/Source/Runtime/Json/Public/Serialization/JsonWriter.h
--- a/Engine/Source/Runtime/Json/Public/Serialization/JsonWriter.h
+++ b/Engine/Source/Runtime/Json/Public/Serialization/JsonWriter.h
@@ -2,11 +2,24 @@
 
 #pragma once
 
+#include "Containers/Utf8String.h"
 #include "CoreMinimal.h"
 #include "Serialization/JsonTypes.h"
 #include "Policies/PrettyJsonPrintPolicy.h"
 #include "Serialization/MemoryWriter.h"
 
+/**
+* Takes an input source char representing and returns if it is possible to represent in DstChar expected encoding
+* 
+* @param Char source char
+* @return false if character is a control char or if it is out of range of representation in case of converting to ANSICHAR. true otherwise
+**/
+template <typename DstChar, typename SrcChar>
+bool HasDestinationJsonStringCharRepresentation(SrcChar Char)
+{
+	return Char >= CHARTEXT(SrcChar, ' ') && (std::is_same_v<DstChar, SrcChar> || !std::is_same_v<DstChar, ANSICHAR> || Char <= (SrcChar)0x79);
+}
+
 /**
  * Takes an input string and escapes it so it can be written as a valid Json string. Also adds the quotes.
  * Appends to a given string-like object to avoid reallocations.
@@ -19,31 +32,32 @@
 template<typename StringType>
 inline StringType& AppendEscapeJsonString(StringType& AppendTo, const FString& StringVal)
 {
+	using CharType = typename StringType::ElementType;
 	AppendTo += TEXT("\"");
-	for (const TCHAR* Char = *StringVal; *Char != TCHAR('\0'); ++Char)
+	for (const CharType* Char = *StringVal; *Char != CHARTEXT(CharType, '\0'); ++Char)
 	{
 		switch (*Char)
 		{
-		case TCHAR('\\'): AppendTo += TEXT("\\\\"); break;
-		case TCHAR('\n'): AppendTo += TEXT("\\n"); break;
-		case TCHAR('\t'): AppendTo += TEXT("\\t"); break;
-		case TCHAR('\b'): AppendTo += TEXT("\\b"); break;
-		case TCHAR('\f'): AppendTo += TEXT("\\f"); break;
-		case TCHAR('\r'): AppendTo += TEXT("\\r"); break;
-		case TCHAR('\"'): AppendTo += TEXT("\\\""); break;
+		case CHARTEXT(CharType, '\\'): AppendTo += CHARTEXT(CharType, "\\\\"); break;
+		case CHARTEXT(CharType, '\n'): AppendTo += CHARTEXT(CharType, "\\n"); break;
+		case CHARTEXT(CharType, '\t'): AppendTo += CHARTEXT(CharType, "\\t"); break;
+		case CHARTEXT(CharType, '\b'): AppendTo += CHARTEXT(CharType, "\\b"); break;
+		case CHARTEXT(CharType, '\f'): AppendTo += CHARTEXT(CharType, "\\f"); break;
+		case CHARTEXT(CharType, '\r'): AppendTo += CHARTEXT(CharType, "\\r"); break;
+		case CHARTEXT(CharType, '\"'): AppendTo += CHARTEXT(CharType, "\\\""); break;
 		default:
-			// Must escape control characters
-			if (*Char >= TCHAR(32))
+			// Must escape control characters or non representable characters
+			if (HasDestinationJsonStringCharRepresentation<CharType>(*Char) )
 			{
 				AppendTo += *Char;
 			}
 			else
 			{
-				AppendTo.Appendf(TEXT("\\u%04x"), *Char);
+				AppendTo.Appendf(CHARTEXT(CharType, "\\u%04x"), *Char);
 			}
 		}
 	}
-	AppendTo += TEXT("\"");
+	AppendTo += CHARTEXT(CharType, "\"");
 
 	return AppendTo;
 }
@@ -261,12 +275,14 @@ public:
 	// Use this to insert the results of a different JSON Writer in.
 	void WriteRawJSONValue( FStringView Identifier, FStringView Value )
 	{
-		check( Stack.Top() == EJson::Object );
-		WriteIdentifier( Identifier );
+		WriteRawJSONValueImpl(Identifier, Value);
+	}
 
-		PrintPolicy::WriteSpace(Stream);
-		PrintPolicy::WriteString(Stream, Value);
-		PreviousTokenWritten = EJsonToken::String;
+	// WARNING: THIS IS DANGEROUS. Use this only if you know for a fact that the Value is valid JSON!
+	// Use this to insert the results of a different JSON Writer in.
+	void WriteRawJSONValue( FUtf8StringView Identifier, FUtf8StringView Value )
+	{
+		WriteRawJSONValueImpl(Identifier, Value);
 	}
 
 	template<typename IdentifierType>
@@ -284,21 +300,14 @@ public:
 	// Use this to insert the results of a different JSON Writer in.
 	void WriteRawJSONValue(FStringView Value)
 	{
-		check(CanWriteValueWithoutIdentifier());
-		WriteCommaIfNeeded();
-
-		if ( PreviousTokenWritten != EJsonToken::True && PreviousTokenWritten != EJsonToken::False && PreviousTokenWritten != EJsonToken::SquareOpen )
-		{
-			PrintPolicy::WriteLineTerminator(Stream);
-			PrintPolicy::WriteTabs(Stream, IndentLevel);
-		}
-		else
-		{
-			PrintPolicy::WriteSpace( Stream );
-		}
+		WriteRawJSONValueImpl(Value);
+	}
 
-		PrintPolicy::WriteString(Stream, Value);
-		PreviousTokenWritten = EJsonToken::String;
+	// WARNING: THIS IS DANGEROUS. Use this only if you know for a fact that the Value is valid JSON!
+	// Use this to insert the results of a different JSON Writer in.
+	void WriteRawJSONValue(FUtf8StringView Value)
+	{
+		WriteRawJSONValueImpl(Value);
 	}
 
 	void WriteNull()
@@ -361,27 +370,19 @@ protected:
 		}
 	}
 
-	void WriteIdentifier(const ANSICHAR* Identifier)
+	template <typename InCharType>
+	void WriteIdentifier(const InCharType* Identifier)
 	{
 		WriteCommaIfNeeded();
 		PrintPolicy::WriteLineTerminator(Stream);
 
 		PrintPolicy::WriteTabs(Stream, IndentLevel);
-		WriteStringValue(FAnsiStringView(Identifier));
+		WriteStringValue(TStringView<InCharType>(Identifier));
 		PrintPolicy::WriteChar(Stream, CharType(':'));
 	}
 	
-	void WriteIdentifier(const TCHAR* Identifier)
-	{
-		WriteCommaIfNeeded();
-		PrintPolicy::WriteLineTerminator(Stream);
-
-		PrintPolicy::WriteTabs(Stream, IndentLevel);
-		WriteStringValue(FStringView(Identifier));
-		PrintPolicy::WriteChar(Stream, CharType(':'));
-	}
-
-	void WriteIdentifier(FStringView Identifier)
+	template <typename InCharType>
+	void WriteIdentifier(TStringView<InCharType> Identifier)
 	{
 		WriteCommaIfNeeded();
 		PrintPolicy::WriteLineTerminator(Stream);
@@ -466,6 +467,12 @@ protected:
 		return EJsonToken::String;
 	}
 
+	FORCEINLINE EJsonToken WriteValueOnly(FUtf8StringView Value)
+	{
+		WriteStringValue(Value);
+		return EJsonToken::String;
+	}
+
 	virtual void WriteStringValue(FAnsiStringView String)
 	{
 		PrintPolicy::WriteChar(Stream, CharType('"'));
@@ -485,6 +492,51 @@ protected:
 		TJsonWriter::WriteStringValue(FStringView(String));
 	}
 
+	virtual void WriteStringValue(FUtf8StringView String)
+	{
+		PrintPolicy::WriteChar(Stream, CharType('"'));
+		WriteEscapedString(String);
+		PrintPolicy::WriteChar(Stream, CharType('"'));
+	}
+
+	virtual void WriteStringValue(const FUtf8String& String)
+	{
+		TJsonWriter::WriteStringValue(FUtf8StringView(String));
+	}
+
+	// WARNING: THIS IS DANGEROUS. Use this only if you know for a fact that the Value is valid JSON!
+	// Use this to insert the results of a different JSON Writer in.
+	template <typename InCharType>
+	void WriteRawJSONValueImpl(FStringView Identifier, TStringView<InCharType> Value)
+	{
+		check(Stack.Top() == EJson::Object);
+		WriteIdentifier(Identifier);
+
+		PrintPolicy::WriteSpace(Stream);
+		PrintPolicy::WriteString(Stream, Value);
+		PreviousTokenWritten = EJsonToken::String;
+	}
+
+	template <typename InCharType>
+	void WriteRawJSONValueImpl(TStringView<InCharType> Value)
+	{
+		check(CanWriteValueWithoutIdentifier());
+		WriteCommaIfNeeded();
+
+		if (PreviousTokenWritten != EJsonToken::True && PreviousTokenWritten != EJsonToken::False && PreviousTokenWritten != EJsonToken::SquareOpen)
+		{
+			PrintPolicy::WriteLineTerminator(Stream);
+			PrintPolicy::WriteTabs(Stream, IndentLevel);
+		}
+		else
+		{
+			PrintPolicy::WriteSpace(Stream);
+		}
+
+		PrintPolicy::WriteString(Stream, Value);
+		PreviousTokenWritten = EJsonToken::String;
+	}
+
 	template<typename InCharType>
 	void WriteEscapedString(TStringView<InCharType> InView)
 	{
@@ -492,16 +544,16 @@ protected:
 		{
 			switch (Char)
 			{
-			case TCHAR('\\'): return true;
-			case TCHAR('\n'): return true;
-			case TCHAR('\t'): return true;
-			case TCHAR('\b'): return true;
-			case TCHAR('\f'): return true;
-			case TCHAR('\r'): return true;
-			case TCHAR('\"'): return true;
+			case CHARTEXT(InCharType, '\\'): return true;
+			case CHARTEXT(InCharType, '\n'): return true;
+			case CHARTEXT(InCharType, '\t'): return true;
+			case CHARTEXT(InCharType, '\b'): return true;
+			case CHARTEXT(InCharType, '\f'): return true;
+			case CHARTEXT(InCharType, '\r'): return true;
+			case CHARTEXT(InCharType, '\"'): return true;
 			default:
-				// Must escape control characters
-				if (Char >= TCHAR(32))
+				// Must escape control characters or non representable characters
+				if (HasDestinationJsonStringCharRepresentation<CharType>(Char))
 				{
 					return false;
 				}
@@ -532,21 +584,21 @@ protected:
 			InView.RightChopInline(EndIndex);
 			for (EndIndex = 0; EndIndex < InView.Len(); ++EndIndex)
 			{
-				TCHAR Char = InView[EndIndex]; 
+				InCharType Char = InView[EndIndex];
 				switch (Char)
 				{
-				case TCHAR('\\'): PrintPolicy::WriteString(Stream, TEXTVIEW("\\\\")); continue;
-				case TCHAR('\n'): PrintPolicy::WriteString(Stream, TEXTVIEW("\\n")); continue;
-				case TCHAR('\t'): PrintPolicy::WriteString(Stream, TEXTVIEW("\\t")); continue;
-				case TCHAR('\b'): PrintPolicy::WriteString(Stream, TEXTVIEW("\\b")); continue;
-				case TCHAR('\f'): PrintPolicy::WriteString(Stream, TEXTVIEW("\\f")); continue;
-				case TCHAR('\r'): PrintPolicy::WriteString(Stream, TEXTVIEW("\\r")); continue;
-				case TCHAR('\"'): PrintPolicy::WriteString(Stream, TEXTVIEW("\\\"")); continue;
+				case CHARTEXT(InCharType, '\\'): PrintPolicy::WriteString(Stream, TEXTVIEW("\\\\")); continue;
+				case CHARTEXT(InCharType, '\n'): PrintPolicy::WriteString(Stream, TEXTVIEW("\\n")); continue;
+				case CHARTEXT(InCharType, '\t'): PrintPolicy::WriteString(Stream, TEXTVIEW("\\t")); continue;
+				case CHARTEXT(InCharType, '\b'): PrintPolicy::WriteString(Stream, TEXTVIEW("\\b")); continue;
+				case CHARTEXT(InCharType, '\f'): PrintPolicy::WriteString(Stream, TEXTVIEW("\\f")); continue;
+				case CHARTEXT(InCharType, '\r'): PrintPolicy::WriteString(Stream, TEXTVIEW("\\r")); continue;
+				case CHARTEXT(InCharType, '\"'): PrintPolicy::WriteString(Stream, TEXTVIEW("\\\"")); continue;
 				default: break;
 				}
 
-				// Must escape control characters
-				if (Char >= TCHAR(32))
+				// Must escape control characters or non representable characters
+				if (HasDestinationJsonStringCharRepresentation<CharType>(Char))
 				{
 					break;
 				}
@@ -570,11 +622,13 @@ protected:
 
 template <class PrintPolicy = TPrettyJsonPrintPolicy<TCHAR>>
 class TJsonStringWriter
-	: public TJsonWriter<TCHAR, PrintPolicy>
+	: public TJsonWriter<typename PrintPolicy::CharType, PrintPolicy>
 {
 public:
+	using CharType = typename PrintPolicy::CharType;
+	using StringType = TString<CharType>;
 
-	static TSharedRef<TJsonStringWriter> Create( FString* const InStream, int32 InitialIndent = 0 )
+	static TSharedRef<TJsonStringWriter> Create(StringType* const InStream, int32 InitialIndent = 0 )
 	{
 		return MakeShareable(new TJsonStringWriter(InStream, InitialIndent));
 	}
@@ -589,20 +643,20 @@ public:
 
 	virtual bool Close() override
 	{
-		OutString->Reset(Bytes.Num()/sizeof(TCHAR));
-		for (int32 i = 0; i < Bytes.Num(); i+=sizeof(TCHAR))
+		OutString->Reset(Bytes.Num()/sizeof(CharType));
+		for (int32 i = 0; i < Bytes.Num(); i+=sizeof(CharType))
 		{
-			TCHAR* Char = static_cast<TCHAR*>(static_cast<void*>(&Bytes[i]));
+			CharType* Char = static_cast<CharType*>(static_cast<void*>(&Bytes[i]));
 			*OutString += *Char;
 		}
 
-		return TJsonWriter<TCHAR, PrintPolicy>::Close();
+		return TJsonWriter<CharType, PrintPolicy>::Close();
 	}
 
 protected:
 
-	TJsonStringWriter( FString* const InOutString, int32 InitialIndent )
-		: TJsonWriter<TCHAR, PrintPolicy>(new FMemoryWriter(Bytes), InitialIndent)
+	TJsonStringWriter(StringType* const InOutString, int32 InitialIndent )
+		: TJsonWriter<CharType, PrintPolicy>(new FMemoryWriter(Bytes), InitialIndent)
 		, Bytes()
 		, OutString(InOutString)
 	{ }
@@ -610,22 +664,22 @@ protected:
 private:
 
 	TArray<uint8> Bytes;
-	FString* OutString;
+	StringType* OutString;
 };
 
-
 template <class CharType = TCHAR, class PrintPolicy = TPrettyJsonPrintPolicy<CharType>>
 class TJsonWriterFactory
 {
 public:
+	using StringType = TString<CharType>;
 
-	static TSharedRef<TJsonWriter<CharType, PrintPolicy>> Create( FArchive* const Stream, int32 InitialIndent = 0 )
+	static TSharedRef<TJsonWriter<CharType, PrintPolicy>> Create(FArchive* const Stream, int32 InitialIndent = 0)
 	{
 		return TJsonWriter< CharType, PrintPolicy >::Create(Stream, InitialIndent);
 	}
 
-	static TSharedRef<TJsonWriter<TCHAR, PrintPolicy>> Create( FString* const Stream, int32 InitialIndent = 0 )
+	static TSharedRef<TJsonWriter<CharType, PrintPolicy>> Create(StringType* const Stream, int32 InitialIndent = 0)
 	{
-		return StaticCastSharedRef<TJsonWriter<TCHAR, PrintPolicy>>(TJsonStringWriter<PrintPolicy>::Create(Stream, InitialIndent));
+		return StaticCastSharedRef<TJsonWriter<CharType, PrintPolicy>>(TJsonStringWriter<PrintPolicy>::Create(Stream, InitialIndent));
 	}
 };
