Fix wrong condition to check whether a character need to be scaped when converted to ANSICHAR

diff --git a/Engine/Source/Runtime/Json/Public/Serialization/JsonWriter.h b/Engine/Source/Runtime/Json/Public/Serialization/JsonWriter.h
--- a/Engine/Source/Runtime/Json/Public/Serialization/JsonWriter.h
+++ b/Engine/Source/Runtime/Json/Public/Serialization/JsonWriter.h
@@ -17,7 +17,7 @@
 template <typename DstChar, typename SrcChar>
 bool HasDestinationJsonStringCharRepresentation(SrcChar Char)
 {
-	return Char >= CHARTEXT(SrcChar, ' ') && (std::is_same_v<DstChar, SrcChar> || !std::is_same_v<DstChar, ANSICHAR> || Char <= (SrcChar)0x79);
+	return Char >= CHARTEXT(SrcChar, ' ') && (std::is_same_v<DstChar, SrcChar> || !std::is_same_v<DstChar, ANSICHAR> || Char <= (SrcChar)0x7E);
 }
 
 /**
