Modifying JsonWriter to not put a comma if there has been no data added. Allows a user to writes non JsonObjects (like strings and numbers). Without this change this:
FJsonSerializer::Serialize(JsonValue, FString(), JsonWriter);
would cause the JsonWriter to add a prefix comma:
,"The Value in JsonValue"

Adds EJsonObjectConversionFlags::SuppressClassNameForPersistentObject to force class name to not be written in each object. When using FJsonSerializer::Serialize() to serialize data there are times it is desired to not have the name of the class in the resulting Json.

[RN] fix Engine/Json

diff --git a/Engine/Source/Runtime/Json/Public/Serialization/JsonWriter.h b/Engine/Source/Runtime/Json/Public/Serialization/JsonWriter.h
--- a/Engine/Source/Runtime/Json/Public/Serialization/JsonWriter.h
+++ b/Engine/Source/Runtime/Json/Public/Serialization/JsonWriter.h
@@ -109,10 +109,7 @@ public:
 	void WriteObjectStart()
 	{
 		check(CanWriteObjectWithoutIdentifier());
-		if (PreviousTokenWritten != EJsonToken::None )
-		{
-			WriteCommaIfNeeded();
-		}
+		WriteCommaIfNeeded();
 
 		if ( PreviousTokenWritten != EJsonToken::None )
 		{
@@ -156,10 +153,7 @@ public:
 	void WriteArrayStart()
 	{
 		check(CanWriteValueWithoutIdentifier());
-		if ( PreviousTokenWritten != EJsonToken::None )
-		{
-			WriteCommaIfNeeded();
-		}
+		WriteCommaIfNeeded();
 
 		if ( PreviousTokenWritten != EJsonToken::None )
 		{
@@ -364,7 +358,7 @@ protected:
 
 	FORCEINLINE void WriteCommaIfNeeded()
 	{
-		if ( PreviousTokenWritten != EJsonToken::CurlyOpen && PreviousTokenWritten != EJsonToken::SquareOpen && PreviousTokenWritten != EJsonToken::Identifier)
+		if ( PreviousTokenWritten != EJsonToken::CurlyOpen && PreviousTokenWritten != EJsonToken::SquareOpen && PreviousTokenWritten != EJsonToken::Identifier && PreviousTokenWritten != EJsonToken::None)
 		{
 			PrintPolicy::WriteChar(Stream, CharType(','));
 		}
diff --git a/Engine/Source/Runtime/JsonUtilities/Private/JsonObjectConverter.cpp b/Engine/Source/Runtime/JsonUtilities/Private/JsonObjectConverter.cpp
--- a/Engine/Source/Runtime/JsonUtilities/Private/JsonObjectConverter.cpp
+++ b/Engine/Source/Runtime/JsonUtilities/Private/JsonObjectConverter.cpp
@@ -187,8 +187,11 @@ TSharedPtr<FJsonValue> ConvertScalarFPropertyToJsonValue(FProperty* Property, co
 		if (Object && (ObjectProperty->HasAnyPropertyFlags(CPF_PersistentInstance) || (OuterProperty && OuterProperty->HasAnyPropertyFlags(CPF_PersistentInstance))))
 		{
 			TSharedRef<FJsonObject> Out = MakeShared<FJsonObject>();
+			if (!EnumHasAnyFlags(ConversionFlags, EJsonObjectConversionFlags::SuppressClassNameForPersistentObject))
+			{
+				Out->SetStringField(ObjectClassNameKey, Object->GetClass()->GetPathName());
+			}
 
-			Out->SetStringField(ObjectClassNameKey, Object->GetClass()->GetPathName());
 			if (FJsonObjectConverter::UStructToJsonObject(ObjectProperty->GetObjectPropertyValue(Value)->GetClass(), Object, Out, CheckFlags, SkipFlags, ExportCb, ConversionFlags))
 			{
 				TSharedRef<FJsonValueObject> JsonObject = MakeShared<FJsonValueObject>(Out);
diff --git a/Engine/Source/Runtime/JsonUtilities/Public/JsonObjectConverter.h b/Engine/Source/Runtime/JsonUtilities/Public/JsonObjectConverter.h
--- a/Engine/Source/Runtime/JsonUtilities/Public/JsonObjectConverter.h
+++ b/Engine/Source/Runtime/JsonUtilities/Public/JsonObjectConverter.h
@@ -36,7 +36,14 @@ enum class EJsonObjectConversionFlags
 	 * Write text in its complex exported format (eg, NSLOCTEXT(...)) rather than as a simple string.
 	 * @note This is required to correctly support localization
 	 */
-	WriteTextAsComplexString = 1 << 1
+	WriteTextAsComplexString = 1 << 1,
+
+	/**
+	 * suppress _ClassName key written for persistent objects that allows them to be
+	 * deserialized later. This can be useful for times when writing to Json and you don't
+	 * need to reserialize back into UObjects, as the _ClassName would be unnecessary.
+	*/
+	SuppressClassNameForPersistentObject = 1 << 2,
 };
 
 ENUM_CLASS_FLAGS(EJsonObjectConversionFlags)
