Move some implementations from JsonValue.h to JsonValue.cpp

diff --git a/Engine/Source/Runtime/Json/Private/Dom/JsonValue.cpp b/Engine/Source/Runtime/Json/Private/Dom/JsonValue.cpp
--- a/Engine/Source/Runtime/Json/Private/Dom/JsonValue.cpp
+++ b/Engine/Source/Runtime/Json/Private/Dom/JsonValue.cpp
@@ -365,7 +365,337 @@ void FJsonValue::ErrorMessage(const FString& InType) const
 	}
 }
 
+
+namespace 
+{
+	template <typename CharType>
+	struct TJsonValueStringType
+	{
+		static_assert(sizeof(CharType) == 0, "Unsupported type");
+	};
+
+	template <>
+	struct TJsonValueStringType<TCHAR>
+	{
+		static FString GetType()
+		{
+			return TEXT("String");
+		}
+	};
+
+	template <>
+	struct TJsonValueStringType<UTF8CHAR>
+	{
+		static FString GetType()
+		{
+			return TEXT("Utf8String");
+		}
+	};
+}
+
+template <typename CharType>
+bool TJsonValueString<CharType>::TryGetString(FString& OutString) const
+{
+	if constexpr (std::is_same_v<FString, TString<CharType>>)
+	{
+		OutString = Value;
+	}
+	else
+	{
+		OutString = FString(Value);
+	}
+	return true;
+}
+
+template <typename CharType>
+bool TJsonValueString<CharType>::TryGetUtf8String(FUtf8String& OutString) const
+{
+	if constexpr (std::is_same_v<FUtf8String, TString<CharType>>)
+	{
+		OutString = Value;
+	}
+	else
+	{
+		OutString = FUtf8String(Value);
+	}
+	return true;
+}
+
+template <typename CharType>
+bool TJsonValueString<CharType>::TryGetNumber(double& OutDouble) const
+{
+	if (Value.IsNumeric())
+	{
+		OutDouble = TCString<CharType>::Atod(*Value);
+		return true;
+	}
+	else
+	{
+		return false;
+	}
+}
+
+template <typename CharType>
+FString TJsonValueString<CharType>::GetType() const
+{
+	return TJsonValueStringType<CharType>::GetType();
+}
+
+template class TJsonValueString<TCHAR>;
+template class TJsonValueString<UTF8CHAR>;
+
+namespace
+{
+template <typename CharType>
+struct TJsonValueNumberType
+{
+	static_assert(sizeof(CharType), "Unsupported type");
+};
+
+template <>
+struct TJsonValueNumberType<TCHAR>
+{
+	static FString GetType()
+	{
+		return TEXT("NumberString");
+	}
+};
+
+template <>
+struct TJsonValueNumberType<UTF8CHAR>
+{
+	static FString GetType()
+	{
+		return TEXT("Utf8NumberString");
+	}
+};
+}
+
+template <typename CharType>
+bool TJsonValueNumberString<CharType>::TryGetString(FString& OutString) const
+{
+	if constexpr (std::is_same_v<FString, TString<CharType>>)
+	{
+		OutString = Value;
+	}
+	else
+	{
+		OutString = FString(Value);
+	}
+	return true;
+}
+
+template <typename CharType>
+bool TJsonValueNumberString<CharType>::TryGetUtf8String(FUtf8String& OutString) const
+{
+	if constexpr (std::is_same_v<FUtf8String, TString<CharType>>)
+	{
+		OutString = Value;
+	}
+	else
+	{
+		OutString = FUtf8String(Value);
+	}
+	return true;
+}
+
+template <typename CharType>
+FString TJsonValueNumberString<CharType>::GetType() const
+{
+	return TJsonValueNumberType<CharType>::GetType();
+}
+
+template class TJsonValueNumberString<TCHAR>;
+template class TJsonValueNumberString<UTF8CHAR>;
+
+SIZE_T FJsonValueArray::GetAllocatedSize() const
+{
+	SIZE_T SizeBytes = 0;
+	SizeBytes += Value.GetAllocatedSize();
+	for (const TSharedPtr<FJsonValue>& Element : Value)
+	{
+		SizeBytes += Element.IsValid() ? Element->GetMemoryFootprint() : 0;
+	}
+	return SizeBytes;
+}
+
 SIZE_T FJsonValueObject::GetAllocatedSize() const
 {
 	return Value.IsValid() ? Value->GetMemoryFootprint() : 0;
 }
+
+
+namespace UE::Json
+{
+	JsonSimpleValueVariant ToSimpleJsonVariant(const FJsonValue& InJsonValue)
+	{
+		if (!InJsonValue.PreferStringRepresentation())
+		{
+			if (InJsonValue.Type == EJson::Boolean)
+			{
+				return JsonSimpleValueVariant(TInPlaceType<bool>(), InJsonValue.AsBool());
+			}
+			else if (InJsonValue.Type == EJson::Number)
+			{
+				const double JsonNumber_v = InJsonValue.AsNumber();
+
+				/* If the Json Number Value requires a decimal point, then we read in the value as a double, otherwise, we read it in as an int */
+				if (FString::SanitizeFloat(JsonNumber_v, 0).Contains(TEXT(".")))
+				{
+					return ToSimpleJsonVariant(JsonNumber_v);
+				}
+				else
+				{
+					return ToSimpleJsonVariant(FMath::RoundToInt64(JsonNumber_v));
+				}
+			}
+		}
+
+		return JsonSimpleValueVariant(TInPlaceType<FString>(), InJsonValue.AsString());
+	}
+}
+
+bool operator==(const JsonNumberValueVariants& Lhs, const FString& Rhs)
+{
+	return Rhs.IsNumeric() && ::Visit([Rhs](const auto& StoredNumber)
+		{
+			using StoredNumberType = std::decay_t<decltype(StoredNumber)>;
+			if constexpr (std::is_same_v<StoredNumberType, float> || std::is_same_v<StoredNumberType, double>)
+			{
+				return FString::SanitizeFloat(StoredNumber, 0) == Rhs;
+			}
+			else
+			{
+				return StoredNumber == FCString::Atoi64(*Rhs);
+			}
+		}, Lhs);
+}
+
+
+FString ToString(const JsonNumberValueVariants& InNumberVariant)
+{
+	return ::Visit([](auto& StoredNumber)
+		{
+			using StoredNumberType = std::decay_t<decltype(StoredNumber)>;
+			if constexpr (std::is_same_v<StoredNumberType, float> || std::is_same_v<StoredNumberType, double>)
+			{
+				return FString::SanitizeFloat(StoredNumber, 0);
+			}
+			else
+			{
+				return FString::Printf(TEXT("%lld"), static_cast<int64>(StoredNumber));
+			}
+		}, InNumberVariant);
+}
+
+bool operator==(const JsonNumberValueVariants& Lhs, const JsonNumberValueVariants& Rhs)
+{
+	const bool bLhsIsFloat = Lhs.IsType<float>() || Lhs.IsType<double>();
+	const bool bRhsIsFloat = Rhs.IsType<float>() || Rhs.IsType<double>();
+	if (bLhsIsFloat || bRhsIsFloat)
+	{
+		return ToString(Lhs) == ToString(Rhs);
+	}
+	else
+	{
+		auto CastToInt64Functor = [](auto& StoredNumber)
+			{
+				return static_cast<int64>(StoredNumber);
+			};
+
+		const int64 LhsValue = ::Visit(CastToInt64Functor, Lhs);
+		const int64 RhsValue = ::Visit(CastToInt64Functor, Rhs);
+
+		return LhsValue == RhsValue;
+	}
+}
+
+
+bool operator==(const JsonSimpleValueVariant& Lhs, const JsonSimpleValueVariant& Rhs)
+{
+	if (Lhs.IsType<bool>())
+	{
+		if (Rhs.IsType<bool>())
+		{
+			return Lhs.Get<bool>() == Rhs.Get<bool>();
+		}
+		else if (Rhs.IsType<FString>())
+		{
+			if (Lhs.Get<bool>())
+			{
+				return Rhs.Get<FString>().Equals(TEXT("true"), ESearchCase::IgnoreCase) ||
+					Rhs.Get<FString>().Equals(TEXT("1"), ESearchCase::IgnoreCase);
+			}
+			else
+			{
+				return Rhs.Get<FString>().Equals(TEXT("false"), ESearchCase::IgnoreCase) ||
+					Rhs.Get<FString>().Equals(TEXT("0"), ESearchCase::IgnoreCase);
+			}
+		}
+		else // RhsType.IsType<JsonNumberValueVariants>()
+		{
+			return ::Visit([&Lhs](const auto& RhsStoredNumber)
+				{
+					using RhsStoredNumberType = std::decay_t<decltype(RhsStoredNumber)>;
+					if constexpr (std::is_same_v<RhsStoredNumberType, float> || std::is_same_v<RhsStoredNumberType, double>)
+					{
+						if (!FString::SanitizeFloat(RhsStoredNumber, 0).Contains(TEXT(".")))
+						{
+							const int64 RhsStoredNumberAsInt = FMath::RoundToInt64(RhsStoredNumber);
+							if (Lhs.Get<bool>())
+							{
+								return RhsStoredNumberAsInt == 1;
+							}
+							else
+							{
+								return RhsStoredNumberAsInt == 0;
+							}
+						}
+						else
+						{
+							return false;
+						}
+					}
+					else
+					{
+						if (Lhs.Get<bool>())
+						{
+							return RhsStoredNumber == 1;
+						}
+						else
+						{
+							return RhsStoredNumber == 0;
+						}
+					}
+				}, Rhs.Get<JsonNumberValueVariants>());
+		}
+	}
+	else if (Lhs.IsType<JsonNumberValueVariants>())
+	{
+		if (Rhs.IsType<JsonNumberValueVariants>())
+		{
+			return Lhs.Get<JsonNumberValueVariants>() == Rhs.Get<JsonNumberValueVariants>();
+		}
+		else // RhsType.IsType<bool>() || RhsType.IsType<FString>()
+		{
+			// Swapping args to avoid code duplication
+			return Rhs == Lhs;
+		}
+	}
+	else // Lhs.IsType<FString>()
+	{
+		if (Rhs.IsType<FString>())
+		{
+			return Lhs.Get<FString>().Equals(Rhs.Get<FString>(), ESearchCase::CaseSensitive);
+		}
+		else if (Rhs.IsType<bool>())
+		{
+			// Swapping args to avoid code duplication
+			return Rhs == Lhs;
+		}
+		else // RhsType.IsType<JsonNumberValueVariants>()
+		{
+			return Lhs.Get<FString>() == Rhs.Get<JsonNumberValueVariants>();
+		}
+	}
+}
\ No newline at end of file
diff --git a/Engine/Source/Runtime/Json/Public/Dom/JsonValue.h b/Engine/Source/Runtime/Json/Public/Dom/JsonValue.h
--- a/Engine/Source/Runtime/Json/Public/Dom/JsonValue.h
+++ b/Engine/Source/Runtime/Json/Public/Dom/JsonValue.h
@@ -118,6 +118,7 @@ public:
 	static JSON_API bool CompareEqual(const FJsonValue& Lhs, const FJsonValue& Rhs);
 
 protected:
+	UE_NONCOPYABLE(FJsonValue)
 
 	FJsonValue() : Type(EJson::None) {}
 	virtual ~FJsonValue() {}
@@ -137,30 +138,6 @@ protected:
 	}
 };
 
-template <typename CharType>
-struct TJsonValueStringType
-{
-	static_assert(sizeof(CharType) == 0, "Unsupported type");
-};
-
-template <>
-struct TJsonValueStringType<TCHAR>
-{
-	static FString GetType()
-	{
-		return TEXT("String");
-	}
-};
-
-template <>
-struct TJsonValueStringType<UTF8CHAR>
-{
-	static FString GetType()
-	{
-		return TEXT("Utf8String");
-	}
-};
-
 /** A Json String Value. */
 template <typename CharType>
 class TJsonValueString : public FJsonValue
@@ -169,44 +146,9 @@ public:
 	TJsonValueString(const TString<CharType>& InString) : Value(InString) {Type = EJson::String;}
 	TJsonValueString(TString<CharType>&& InString) : Value(MoveTemp(InString)) {Type = EJson::String;}
 
-	virtual bool TryGetString(FString& OutString) const override
-	{
-		if constexpr (std::is_same_v<FString, TString<CharType>>)
-		{
-			OutString = Value;
-		}
-		else
-		{
-			OutString = FString(Value);
-		}
-		return true; 
-	}
-	
-	virtual bool TryGetUtf8String(FUtf8String& OutString) const override
-	{
-		if constexpr (std::is_same_v<FUtf8String, TString<CharType>>)
-		{
-			OutString = Value;
-		}
-		else
-		{
-			OutString = FUtf8String(Value);
-		}
-		return true;
-	}
-	
-	virtual bool TryGetNumber(double& OutDouble) const override
-	{
-		if (Value.IsNumeric())
-		{
-			OutDouble = TCString<CharType>::Atod(*Value);
-			return true;
-		}
-		else
-		{
-			return false;
-		}
-	}
+	JSON_API virtual bool TryGetString(FString& OutString) const override;	
+	JSON_API virtual bool TryGetUtf8String(FUtf8String& OutString) const override;
+	JSON_API virtual bool TryGetNumber(double& OutDouble) const override;
 
 	virtual bool TryGetNumber(int32& OutValue) const override	{ LexFromString(OutValue, *Value); return true; }
 	
@@ -223,18 +165,22 @@ public:
 
 	virtual SIZE_T GetMemoryFootprint() const override { return sizeof(*this) + GetAllocatedSize(); }
 
+
 protected:
 	TString<CharType> Value;
 
-	virtual FString GetType() const override { return TJsonValueStringType<CharType>::GetType(); }
+	JSON_API virtual FString GetType() const override;
 
 	/** Helper to calculate allocated size of the Value string */
 	SIZE_T GetAllocatedSize() const { return Value.GetAllocatedSize(); }
 };
 
+extern template class TJsonValueString<TCHAR>;
+extern template class TJsonValueString<UTF8CHAR>;
+
 using FJsonValueString = TJsonValueString<TCHAR>;
 
-/** A Json Number Value. */
+//** A Json Number Value. */
 class FJsonValueNumber : public FJsonValue
 {
 public:
@@ -252,30 +198,6 @@ protected:
 	virtual FString GetType() const override {return TEXT("Number");}
 };
 
-template <typename CharType>
-struct TJsonValueNumberType
-{
-	static_assert(sizeof(CharType), "Unsupported type");
-};
-
-template <>
-struct TJsonValueNumberType<TCHAR>
-{
-	static FString GetType()
-	{
-		return TEXT("NumberString");
-	}
-};
-
-template <>
-struct TJsonValueNumberType<UTF8CHAR>
-{
-	static FString GetType()
-	{
-		return TEXT("Utf8NumberString");
-	}
-};
-
 /** A Json Number Value, stored internally as a string so as not to lose precision */
 template <typename CharType>
 class TJsonValueNumberString : public FJsonValue
@@ -285,31 +207,8 @@ public:
 
 	TJsonValueNumberString(TString<CharType>&& InString) : Value(MoveTemp(InString)) { Type = EJson::Number; }
 
-	virtual bool TryGetString(FString& OutString) const override
-	{
-		if constexpr (std::is_same_v<FString, TString<CharType>>)
-		{
-			OutString = Value;
-		}
-		else
-		{
-			OutString = FString(Value);
-		}
-		return true;
-	}
-	
-	virtual bool TryGetUtf8String(FUtf8String& OutString) const override
-	{
-		if constexpr (std::is_same_v<FUtf8String, TString<CharType>>)
-		{
-			OutString = Value;
-		}
-		else
-		{
-			OutString = FUtf8String(Value);
-		}
-		return true;
-	}
+	JSON_API virtual bool TryGetString(FString& OutString) const override;	
+	JSON_API virtual bool TryGetUtf8String(FUtf8String& OutString) const override;
 	
 	virtual bool TryGetNumber(double& OutDouble) const override { return LexTryParseString(OutDouble, *Value); }
 	virtual bool TryGetNumber(float &OutDouble) const override { return LexTryParseString(OutDouble, *Value); }
@@ -325,14 +224,19 @@ public:
 	virtual bool PreferStringRepresentation() const override { return true; }
 	virtual SIZE_T GetMemoryFootprint() const override { return sizeof(*this) + GetAllocatedSize(); }
 
+
 protected:
 	TString<CharType> Value;
 
-	virtual FString GetType() const override { return TJsonValueNumberType<CharType>::GetType(); }
+	JSON_API virtual FString GetType() const override;
+
 	/** Helper to calculate allocated size of the Value string */
 	SIZE_T GetAllocatedSize() const { return Value.GetAllocatedSize(); }
 };
 
+extern template class TJsonValueNumberString<TCHAR>;
+extern template class TJsonValueNumberString<UTF8CHAR>;
+
 using FJsonValueNumberString = TJsonValueNumberString<TCHAR>;
 
 /** A Json Boolean Value. */
@@ -352,7 +256,6 @@ protected:
 	virtual FString GetType() const override {return TEXT("Boolean");}
 };
 
-
 /** A Json Array Value. */
 class FJsonValueArray : public FJsonValue
 {
@@ -364,20 +267,11 @@ public:
 	virtual SIZE_T GetMemoryFootprint() const override { return sizeof(*this) + GetAllocatedSize(); }
 
 protected:
-	TArray< TSharedPtr<FJsonValue> > Value;
+	TArray<TSharedPtr<FJsonValue>> Value;
 
-	virtual FString GetType() const override {return TEXT("Array");}
+	virtual FString GetType() const override {return TEXT("Array");}	
 	/** Helper to calculate allocated size of the Value array and its contents */
-	SIZE_T GetAllocatedSize() const 
-	{
-		SIZE_T SizeBytes = 0;
-		SizeBytes += Value.GetAllocatedSize();
-		for (const TSharedPtr<FJsonValue>& Element : Value)
-		{
-			SizeBytes += Element.IsValid() ? Element->GetMemoryFootprint() : 0;
-		}
-		return SizeBytes; 
-	}
+	JSON_API SIZE_T GetAllocatedSize() const;
 };
 
 
@@ -413,7 +307,7 @@ namespace UE::Json
 {
 
 template<typename T, typename = typename std::enable_if<!std::is_same_v<T, FJsonValue>>>
-static JsonSimpleValueVariant ToSimpleJsonVariant(const T& InSimpleValue)
+inline JsonSimpleValueVariant ToSimpleJsonVariant(const T& InSimpleValue)
 {
 	using InSimpleValueType = std::decay_t<decltype(InSimpleValue)>;
 	if constexpr (std::is_same_v<InSimpleValueType, bool> || std::is_same_v<InSimpleValueType, FString>)
@@ -426,205 +320,42 @@ static JsonSimpleValueVariant ToSimpleJsonVariant(const T& InSimpleValue)
 	}
 }
 
-static JsonSimpleValueVariant ToSimpleJsonVariant(const FJsonValue& InJsonValue)
-{
-	if (!InJsonValue.PreferStringRepresentation())
-	{
-		if (InJsonValue.Type == EJson::Boolean)
-		{
-			return JsonSimpleValueVariant(TInPlaceType<bool>(), InJsonValue.AsBool());
-		}
-		else if (InJsonValue.Type == EJson::Number)
-		{
-			const double JsonNumber_v = InJsonValue.AsNumber();
-
-			/* If the Json Number Value requires a decimal point, then we read in the value as a double, otherwise, we read it in as an int */
-			if (FString::SanitizeFloat(JsonNumber_v, 0).Contains(TEXT(".")))
-			{
-				return ToSimpleJsonVariant(JsonNumber_v);
-			}
-			else
-			{
-				return ToSimpleJsonVariant(FMath::RoundToInt64(JsonNumber_v));
-			}
-		}
-	}
-
-	return JsonSimpleValueVariant(TInPlaceType<FString>(), InJsonValue.AsString());
-}
+JSON_API JsonSimpleValueVariant ToSimpleJsonVariant(const FJsonValue& InJsonValue);
 
 } // namespace UE::Json
 
+JSON_API FString ToString(const JsonNumberValueVariants& InNumberVariant);
+
 /* Global operators */
+JSON_API bool operator==(const JsonSimpleValueVariant& Lhs, const JsonSimpleValueVariant& Rhs);
+JSON_API bool operator==(const JsonNumberValueVariants& Lhs, const JsonNumberValueVariants& Rhs);
+JSON_API bool operator==(const JsonNumberValueVariants& Lhs, const FString& Rhs);
 
-static bool operator==(const JsonNumberValueVariants& Lhs, const FString& Rhs)
-{
-	return Rhs.IsNumeric() && ::Visit([Rhs](const auto& StoredNumber)
-		{
-			using StoredNumberType = std::decay_t<decltype(StoredNumber)>;
-			if constexpr (std::is_same_v<StoredNumberType, float> || std::is_same_v<StoredNumberType, double>)
-			{
-				return FString::SanitizeFloat(StoredNumber, 0) == Rhs;
-			}
-			else
-			{
-				return StoredNumber == FCString::Atoi64(*Rhs);
-			}
-		}, Lhs);
-}
 
-static bool operator!=(const JsonNumberValueVariants& Lhs, const FString& Rhs)
-{
-	return !(Lhs == Rhs);
-}
+#if !PLATFORM_COMPILER_HAS_GENERATED_COMPARISON_OPERATORS
 
-static bool operator==(const FString& Lhs, const JsonNumberValueVariants& Rhs)
+inline bool operator==(const FString& Lhs, const JsonNumberValueVariants& Rhs)
 {
 	return Rhs == Lhs;
 }
 
-static bool operator!=(const FString& Lhs, const JsonNumberValueVariants& Rhs)
+inline bool operator!=(const JsonNumberValueVariants& Lhs, const FString& Rhs)
 {
 	return !(Lhs == Rhs);
 }
 
-////////////////////////////////////////////////////////////
-
-static FString ToString(const JsonNumberValueVariants& InNumberVariant)
-{
-	return ::Visit([](auto& StoredNumber)
-		{
-			using StoredNumberType = std::decay_t<decltype(StoredNumber)>;
-			if constexpr (std::is_same_v<StoredNumberType, float> || std::is_same_v<StoredNumberType, double>)
-			{
-				return FString::SanitizeFloat(StoredNumber, 0);
-			}
-			else
-			{
-				return FString::Printf(TEXT("%lld"), static_cast<int64>(StoredNumber));
-			}
-		}, InNumberVariant);
-}
-
-static bool operator==(const JsonNumberValueVariants& Lhs, const JsonNumberValueVariants& Rhs)
-{
-	const bool bLhsIsFloat = Lhs.IsType<float>() || Lhs.IsType<double>();
-	const bool bRhsIsFloat = Rhs.IsType<float>() || Rhs.IsType<double>();
-	if (bLhsIsFloat || bRhsIsFloat)
-	{
-		return ToString(Lhs) == ToString(Rhs);
-	}
-	else
-	{
-		auto CastToInt64Functor = [](auto& StoredNumber)
-			{
-				return static_cast<int64>(StoredNumber);
-			};
-
-		const int64 LhsValue = ::Visit(CastToInt64Functor, Lhs);
-		const int64 RhsValue = ::Visit(CastToInt64Functor, Rhs);
-
-		return LhsValue == RhsValue;
-	}
-}
-
-static bool operator!=(const JsonNumberValueVariants& Lhs, const JsonNumberValueVariants& Rhs)
+inline bool operator!=(const FString& Lhs, const JsonNumberValueVariants& Rhs)
 {
 	return !(Lhs == Rhs);
 }
 
-////////////////////////////////////////////////////////////
-
-static bool operator==(const JsonSimpleValueVariant& Lhs, const JsonSimpleValueVariant& Rhs)
+inline bool operator!=(const JsonNumberValueVariants& Lhs, const JsonNumberValueVariants& Rhs)
 {
-	if (Lhs.IsType<bool>())
-	{
-		if (Rhs.IsType<bool>())
-		{
-			return Lhs.Get<bool>() == Rhs.Get<bool>();
-		}
-		else if (Rhs.IsType<FString>())
-		{
-			if (Lhs.Get<bool>())
-			{
-				return Rhs.Get<FString>().Equals(TEXT("true"), ESearchCase::IgnoreCase) ||
-					Rhs.Get<FString>().Equals(TEXT("1"), ESearchCase::IgnoreCase);
-			}
-			else
-			{
-				return Rhs.Get<FString>().Equals(TEXT("false"), ESearchCase::IgnoreCase) ||
-					Rhs.Get<FString>().Equals(TEXT("0"), ESearchCase::IgnoreCase);
-			}
-		}
-		else // RhsType.IsType<JsonNumberValueVariants>()
-		{
-			return ::Visit([&Lhs](const auto& RhsStoredNumber)
-				{
-					using RhsStoredNumberType = std::decay_t<decltype(RhsStoredNumber)>;
-					if constexpr (std::is_same_v<RhsStoredNumberType, float> || std::is_same_v<RhsStoredNumberType, double>)
-					{
-						if (!FString::SanitizeFloat(RhsStoredNumber, 0).Contains(TEXT(".")))
-						{
-							const int64 RhsStoredNumberAsInt = FMath::RoundToInt64(RhsStoredNumber);
-							if (Lhs.Get<bool>())
-							{
-								return RhsStoredNumberAsInt == 1;
-							}
-							else
-							{
-								return RhsStoredNumberAsInt == 0;
-							}
-						}
-						else
-						{
-							return false;
-						}
-					}
-					else
-					{
-						if (Lhs.Get<bool>())
-						{
-							return RhsStoredNumber == 1;
-						}
-						else
-						{
-							return RhsStoredNumber == 0;
-						}
-					}
-				}, Rhs.Get<JsonNumberValueVariants>());
-		}
-	}
-	else if (Lhs.IsType<JsonNumberValueVariants>())
-	{
-		if (Rhs.IsType<JsonNumberValueVariants>())
-		{
-			return Lhs.Get<JsonNumberValueVariants>() == Rhs.Get<JsonNumberValueVariants>();
-		}
-		else // RhsType.IsType<bool>() || RhsType.IsType<FString>()
-		{
-			// Swapping args to avoid code duplication
-			return Rhs == Lhs;
-		}
-	}
-	else // Lhs.IsType<FString>()
-	{
-		if (Rhs.IsType<FString>())
-		{
-			return Lhs.Get<FString>().Equals(Rhs.Get<FString>(), ESearchCase::CaseSensitive);
-		}
-		else if (Rhs.IsType<bool>())
-		{
-			// Swapping args to avoid code duplication
-			return Rhs == Lhs;
-		}
-		else // RhsType.IsType<JsonNumberValueVariants>()
-		{
-			return Lhs.Get<FString>() == Rhs.Get<JsonNumberValueVariants>();
-		}
-	}
+	return !(Lhs == Rhs);
 }
 
-static bool operator!=(const JsonSimpleValueVariant& Lhs, const JsonSimpleValueVariant& Rhs)
+inline bool operator!=(const JsonSimpleValueVariant& Lhs, const JsonSimpleValueVariant& Rhs)
 {
 	return !(Lhs == Rhs);
-}
\ No newline at end of file
+}
+#endif
\ No newline at end of file
