Fix missing symbols compiling Mac

diff --git a/Engine/Source/Runtime/Json/Private/Dom/JsonValue.cpp b/Engine/Source/Runtime/Json/Private/Dom/JsonValue.cpp
--- a/Engine/Source/Runtime/Json/Private/Dom/JsonValue.cpp
+++ b/Engine/Source/Runtime/Json/Private/Dom/JsonValue.cpp
@@ -393,6 +393,20 @@ namespace
 	};
 }
 
+template <typename CharType>
+TJsonValueString<CharType>::TJsonValueString(const TString<CharType>& InString)
+: Value(InString)
+{
+	Type = EJson::String;
+}
+
+template <typename CharType>
+TJsonValueString<CharType>::TJsonValueString(TString<CharType>&& InString)
+: Value(MoveTemp(InString))
+{
+	Type = EJson::String;
+}
+
 template <typename CharType>
 bool TJsonValueString<CharType>::TryGetString(FString& OutString) const
 {
@@ -471,6 +485,20 @@ struct TJsonValueNumberType<UTF8CHAR>
 };
 }
 
+template <typename CharType>
+ TJsonValueNumberString<CharType>::TJsonValueNumberString(const TString<CharType>& InString)
+ : Value(InString)
+ {
+	Type = EJson::Number;
+}
+
+template <typename CharType>
+ TJsonValueNumberString<CharType>::TJsonValueNumberString(TString<CharType>&& InString)
+ : Value(MoveTemp(InString))
+ {
+	Type = EJson::Number;
+}
+
 template <typename CharType>
 bool TJsonValueNumberString<CharType>::TryGetString(FString& OutString) const
 {
diff --git a/Engine/Source/Runtime/Json/Public/Dom/JsonValue.h b/Engine/Source/Runtime/Json/Public/Dom/JsonValue.h
--- a/Engine/Source/Runtime/Json/Public/Dom/JsonValue.h
+++ b/Engine/Source/Runtime/Json/Public/Dom/JsonValue.h
@@ -143,8 +143,8 @@ template <typename CharType>
 class TJsonValueString : public FJsonValue
 {
 public:
-	TJsonValueString(const TString<CharType>& InString) : Value(InString) {Type = EJson::String;}
-	TJsonValueString(TString<CharType>&& InString) : Value(MoveTemp(InString)) {Type = EJson::String;}
+	JSON_API TJsonValueString(const TString<CharType>& InString);
+	JSON_API TJsonValueString(TString<CharType>&& InString);
 
 	JSON_API virtual bool TryGetString(FString& OutString) const override;	
 	JSON_API virtual bool TryGetUtf8String(FUtf8String& OutString) const override;
@@ -203,9 +203,8 @@ template <typename CharType>
 class TJsonValueNumberString : public FJsonValue
 {
 public:
-	TJsonValueNumberString(const TString<CharType>& InString) : Value(InString) { Type = EJson::Number; }
-
-	TJsonValueNumberString(TString<CharType>&& InString) : Value(MoveTemp(InString)) { Type = EJson::Number; }
+	JSON_API TJsonValueNumberString(const TString<CharType>& InString);
+	JSON_API TJsonValueNumberString(TString<CharType>&& InString);
 
 	JSON_API virtual bool TryGetString(FString& OutString) const override;	
 	JSON_API virtual bool TryGetUtf8String(FUtf8String& OutString) const override;
