[Engine]
* Fixed clang external template linker errors

diff --git a/Engine/Source/Runtime/Json/Private/Dom/JsonValue.cpp b/Engine/Source/Runtime/Json/Private/Dom/JsonValue.cpp
--- a/Engine/Source/Runtime/Json/Private/Dom/JsonValue.cpp
+++ b/Engine/Source/Runtime/Json/Private/Dom/JsonValue.cpp
@@ -455,8 +455,8 @@ FString TJsonValueString<CharType>::GetType() const
 	return TJsonValueStringType<CharType>::GetType();
 }
 
-template class TJsonValueString<TCHAR>;
-template class TJsonValueString<UTF8CHAR>;
+template class JSON_API TJsonValueString<TCHAR>;
+template class JSON_API TJsonValueString<UTF8CHAR>;
 
 namespace
 {
@@ -533,8 +533,8 @@ FString TJsonValueNumberString<CharType>::GetType() const
 	return TJsonValueNumberType<CharType>::GetType();
 }
 
-template class TJsonValueNumberString<TCHAR>;
-template class TJsonValueNumberString<UTF8CHAR>;
+template class JSON_API TJsonValueNumberString<TCHAR>;
+template class JSON_API TJsonValueNumberString<UTF8CHAR>;
 
 SIZE_T FJsonValueArray::GetAllocatedSize() const
 {
diff --git a/Engine/Source/Runtime/Json/Public/Dom/JsonValue.h b/Engine/Source/Runtime/Json/Public/Dom/JsonValue.h
--- a/Engine/Source/Runtime/Json/Public/Dom/JsonValue.h
+++ b/Engine/Source/Runtime/Json/Public/Dom/JsonValue.h
@@ -143,12 +143,12 @@ template <typename CharType>
 class TJsonValueString : public FJsonValue
 {
 public:
-	JSON_API TJsonValueString(const TString<CharType>& InString);
-	JSON_API TJsonValueString(TString<CharType>&& InString);
+	TJsonValueString(const TString<CharType>& InString);
+	TJsonValueString(TString<CharType>&& InString);
 
-	JSON_API virtual bool TryGetString(FString& OutString) const override;	
-	JSON_API virtual bool TryGetUtf8String(FUtf8String& OutString) const override;
-	JSON_API virtual bool TryGetNumber(double& OutDouble) const override;
+	virtual bool TryGetString(FString& OutString) const override;	
+	virtual bool TryGetUtf8String(FUtf8String& OutString) const override;
+	virtual bool TryGetNumber(double& OutDouble) const override;
 
 	virtual bool TryGetNumber(int32& OutValue) const override	{ LexFromString(OutValue, *Value); return true; }
 	
@@ -169,14 +169,20 @@ public:
 protected:
 	TString<CharType> Value;
 
-	JSON_API virtual FString GetType() const override;
+	virtual FString GetType() const override;
 
 	/** Helper to calculate allocated size of the Value string */
 	SIZE_T GetAllocatedSize() const { return Value.GetAllocatedSize(); }
 };
 
-extern template class TJsonValueString<TCHAR>;
-extern template class TJsonValueString<UTF8CHAR>;
+#if PLATFORM_COMPILER_CLANG
+#define UE_EXTERN_TEMPLATE_API JSON_API
+#else
+#define UE_EXTERN_TEMPLATE_API
+#endif
+
+extern template class UE_EXTERN_TEMPLATE_API TJsonValueString<TCHAR>;
+extern template class UE_EXTERN_TEMPLATE_API TJsonValueString<UTF8CHAR>;
 
 using FJsonValueString = TJsonValueString<TCHAR>;
 
@@ -203,11 +209,11 @@ template <typename CharType>
 class TJsonValueNumberString : public FJsonValue
 {
 public:
-	JSON_API TJsonValueNumberString(const TString<CharType>& InString);
-	JSON_API TJsonValueNumberString(TString<CharType>&& InString);
+	TJsonValueNumberString(const TString<CharType>& InString);
+	TJsonValueNumberString(TString<CharType>&& InString);
 
-	JSON_API virtual bool TryGetString(FString& OutString) const override;	
-	JSON_API virtual bool TryGetUtf8String(FUtf8String& OutString) const override;
+	virtual bool TryGetString(FString& OutString) const override;	
+	virtual bool TryGetUtf8String(FUtf8String& OutString) const override;
 	
 	virtual bool TryGetNumber(double& OutDouble) const override { return LexTryParseString(OutDouble, *Value); }
 	virtual bool TryGetNumber(float &OutDouble) const override { return LexTryParseString(OutDouble, *Value); }
@@ -227,14 +233,16 @@ public:
 protected:
 	TString<CharType> Value;
 
-	JSON_API virtual FString GetType() const override;
+	virtual FString GetType() const override;
 
 	/** Helper to calculate allocated size of the Value string */
 	SIZE_T GetAllocatedSize() const { return Value.GetAllocatedSize(); }
 };
 
-extern template class TJsonValueNumberString<TCHAR>;
-extern template class TJsonValueNumberString<UTF8CHAR>;
+extern template class UE_EXTERN_TEMPLATE_API TJsonValueNumberString<TCHAR>;
+extern template class UE_EXTERN_TEMPLATE_API TJsonValueNumberString<UTF8CHAR>;
+
+#undef UE_EXTERN_TEMPLATE_API
 
 using FJsonValueNumberString = TJsonValueNumberString<TCHAR>;
 
@@ -357,4 +365,4 @@ inline bool operator!=(const JsonSimpleValueVariant& Lhs, const JsonSimpleValueV
 {
 	return !(Lhs == Rhs);
 }
-#endif
\ No newline at end of file
+#endif
