Move all virtual function implementations to JsonValue.cpp

diff --git a/Engine/Source/Runtime/Json/Private/Dom/JsonValue.cpp b/Engine/Source/Runtime/Json/Private/Dom/JsonValue.cpp
--- a/Engine/Source/Runtime/Json/Private/Dom/JsonValue.cpp
+++ b/Engine/Source/Runtime/Json/Private/Dom/JsonValue.cpp
@@ -3,6 +3,15 @@
 #include "Dom/JsonValue.h"
 #include "Dom/JsonObject.h"
 
+FJsonValue::FJsonValue()
+: Type(EJson::None) 
+{
+}
+
+FJsonValue::~FJsonValue()
+{
+}
+
 double FJsonValue::AsNumber() const
 {
 	double Number = 0.0;
@@ -53,13 +62,13 @@ bool FJsonValue::AsBool() const
 }
 
 
-const TArray< TSharedPtr<FJsonValue> >& FJsonValue::AsArray() const
+const TArray<TSharedPtr<FJsonValue>>& FJsonValue::AsArray() const
 {
 	const TArray<TSharedPtr<FJsonValue>>* Array = nullptr;
 
 	if (!TryGetArray(Array))
 	{
-		static const TArray< TSharedPtr<FJsonValue> > EmptyArray;
+		static const TArray<TSharedPtr<FJsonValue>> EmptyArray;
 		Array = &EmptyArray;
 		ErrorMessage(TEXT("Array"));
 	}
@@ -102,7 +111,8 @@ bool TryConvertNumber(const FJsonValue& InValue, T& OutNumber)
 // Need special handling for int64/uint64, due to overflow in the numeric limits.
 // 2^63-1 and 2^64-1 cannot be exactly represented as a double, so TNumericLimits<>::Max() gets rounded up to exactly 2^63 or 2^64 by the compiler's implicit cast to double.
 // This breaks the overflow check in TryConvertNumber. We use "<" rather than "<=" along with the exact power-of-two double literal to fix this.
-template <> bool TryConvertNumber<uint64>(const FJsonValue& InValue, uint64& OutNumber)
+template <> 
+bool TryConvertNumber<uint64>(const FJsonValue& InValue, uint64& OutNumber)
 {
 	double Double;
 	if (InValue.TryGetNumber(Double) && Double >= 0.0 && Double < 18446744073709551616.0)
@@ -114,7 +124,8 @@ template <> bool TryConvertNumber<uint64>(const FJsonValue& InValue, uint64& Out
 	return false;
 }
 
-template <> bool TryConvertNumber<int64>(const FJsonValue& InValue, int64& OutNumber)
+template <>
+bool TryConvertNumber<int64>(const FJsonValue& InValue, int64& OutNumber)
 {
 	double Double;
 	if (InValue.TryGetNumber(Double) && Double >= -9223372036854775808.0 && Double < 9223372036854775808.0)
@@ -127,6 +138,10 @@ template <> bool TryConvertNumber<int64>(const FJsonValue& InValue, int64& OutNu
 }
 
 // -----------------------------------
+bool FJsonValue::TryGetNumber(double& OutNumber) const
+{
+	return false;
+}
 
 bool FJsonValue::TryGetNumber(float& OutNumber) const
 {
@@ -181,6 +196,51 @@ bool FJsonValue::TryGetNumber(int64& OutNumber) const
 	return TryConvertNumber(*this, OutNumber);
 }
 
+bool FJsonValue::TryGetString(FString& OutString) const
+{
+	return false;
+}
+
+bool FJsonValue::TryGetUtf8String(FUtf8String& OutString) const
+{
+	return false;
+}
+
+bool FJsonValue::TryGetBool(bool& OutBool) const
+{
+	return false;
+}
+
+bool FJsonValue::TryGetArray(const TArray<TSharedPtr<FJsonValue>>*& OutArray) const
+{
+	return false;
+}
+	
+bool FJsonValue::TryGetArray(TArray<TSharedPtr<FJsonValue>>*& OutArray)
+{
+	return false;
+}
+
+bool FJsonValue::TryGetObject(const TSharedPtr<FJsonObject>*& Object) const
+{
+	return false;
+}
+
+bool FJsonValue::TryGetObject(TSharedPtr<FJsonObject>*& Object)
+{
+	return false;
+}
+
+bool FJsonValue::PreferStringRepresentation() const
+{
+	return false;
+}
+
+SIZE_T FJsonValue::GetMemoryFootprint() const
+{
+	return sizeof(*this);
+}
+
 //static 
 bool FJsonValue::CompareEqual( const FJsonValue& Lhs, const FJsonValue& Rhs )
 {
@@ -215,8 +275,8 @@ bool FJsonValue::CompareEqual( const FJsonValue& Lhs, const FJsonValue& Rhs )
 
 	case EJson::Array:
 		{
-			const TArray< TSharedPtr<FJsonValue> >& LhsArray = Lhs.AsArray();
-			const TArray< TSharedPtr<FJsonValue> >& RhsArray = Rhs.AsArray();
+			const TArray<TSharedPtr<FJsonValue>>& LhsArray = Lhs.AsArray();
+			const TArray<TSharedPtr<FJsonValue>>& RhsArray = Rhs.AsArray();
 
 			if (LhsArray.Num() != RhsArray.Num())
 			{
@@ -449,15 +509,103 @@ bool TJsonValueString<CharType>::TryGetNumber(double& OutDouble) const
 	}
 }
 
+template <typename CharType>
+bool TJsonValueString<CharType>::TryGetNumber(int32& OutValue) const
+{
+	LexFromString(OutValue, *Value);
+	return true;
+}	
+
+template <typename CharType>
+bool TJsonValueString<CharType>::TryGetNumber(uint32& OutValue) const
+{
+	LexFromString(OutValue, *Value);
+	return true;
+}
+
+template <typename CharType>
+bool TJsonValueString<CharType>::TryGetNumber(int64& OutValue) const
+{
+	LexFromString(OutValue, *Value);
+	return true; 
+}
+
+template <typename CharType>
+bool TJsonValueString<CharType>::TryGetNumber(uint64& OutValue) const
+{
+	LexFromString(OutValue, *Value);
+	return true;
+}
+
+template <typename CharType>
+bool TJsonValueString<CharType>::TryGetBool(bool& OutBool) const
+{
+	OutBool = Value.ToBool();
+	return true;
+}
+
+template <typename CharType>
+SIZE_T TJsonValueString<CharType>::GetMemoryFootprint() const
+{
+	return sizeof(*this) + GetAllocatedSize();
+}
+
 template <typename CharType>
 FString TJsonValueString<CharType>::GetType() const
 {
 	return TJsonValueStringType<CharType>::GetType();
 }
 
+template <typename CharType>
+SIZE_T TJsonValueString<CharType>::GetAllocatedSize() const
+{
+	return Value.GetAllocatedSize();
+}
+
 template class JSON_API TJsonValueString<TCHAR>;
 template class JSON_API TJsonValueString<UTF8CHAR>;
 
+
+FJsonValueNumber::FJsonValueNumber(double InNumber)
+: Value(InNumber)
+{
+	Type = EJson::Number;
+}
+
+bool FJsonValueNumber::TryGetNumber(double& OutNumber) const
+{
+	OutNumber = Value;
+	return true;
+}
+
+bool FJsonValueNumber::TryGetBool(bool& OutBool) const
+{
+	OutBool = (Value != 0.0);
+	return true;
+}
+
+bool FJsonValueNumber::TryGetString(FString& OutString) const
+{
+	OutString = FString::SanitizeFloat(Value, 0);
+	return true;
+}
+
+bool FJsonValueNumber::TryGetUtf8String(FUtf8String& OutString) const
+{
+	OutString = FUtf8String::SanitizeFloat(Value, 0);
+	return true;
+}
+
+SIZE_T FJsonValueNumber::GetMemoryFootprint() const
+{
+	return sizeof(*this);
+}
+
+FString FJsonValueNumber::GetType() const
+{
+	return TEXT("Number");
+}
+
 namespace
 {
 template <typename CharType>
@@ -486,16 +634,16 @@ struct TJsonValueNumberType<UTF8CHAR>
 }
 
 template <typename CharType>
- TJsonValueNumberString<CharType>::TJsonValueNumberString(const TString<CharType>& InString)
- : Value(InString)
- {
+TJsonValueNumberString<CharType>::TJsonValueNumberString(const TString<CharType>& InString)
+: Value(InString)
+{
 	Type = EJson::Number;
 }
 
 template <typename CharType>
- TJsonValueNumberString<CharType>::TJsonValueNumberString(TString<CharType>&& InString)
- : Value(MoveTemp(InString))
- {
+TJsonValueNumberString<CharType>::TJsonValueNumberString(TString<CharType>&& InString)
+: Value(MoveTemp(InString))
+{
 	Type = EJson::Number;
 }
 
@@ -527,15 +675,169 @@ bool TJsonValueNumberString<CharType>::TryGetUtf8String(FUtf8String& OutString)
 	return true;
 }
 
+template <typename CharType>
+bool TJsonValueNumberString<CharType>::TryGetNumber(double& OutDouble) const
+{
+	return LexTryParseString(OutDouble, *Value);
+}
+
+template <typename CharType>
+bool TJsonValueNumberString<CharType>::TryGetNumber(float &OutDouble) const
+{
+	return LexTryParseString(OutDouble, *Value);
+}
+
+template <typename CharType>
+bool TJsonValueNumberString<CharType>::TryGetNumber(int8& OutValue) const
+{
+	return LexTryParseString(OutValue, *Value);
+}
+
+template <typename CharType>
+bool TJsonValueNumberString<CharType>::TryGetNumber(int16& OutValue) const
+{
+	return LexTryParseString(OutValue, *Value);
+}
+
+template <typename CharType>
+bool TJsonValueNumberString<CharType>::TryGetNumber(int32& OutValue) const
+{
+	return LexTryParseString(OutValue, *Value);
+}
+
+template <typename CharType>
+bool TJsonValueNumberString<CharType>::TryGetNumber(int64& OutValue) const
+{
+	return LexTryParseString(OutValue, *Value);
+}
+
+template <typename CharType>
+bool TJsonValueNumberString<CharType>::TryGetNumber(uint8& OutValue) const
+{
+	return LexTryParseString(OutValue, *Value);
+}
+
+template <typename CharType>
+bool TJsonValueNumberString<CharType>::TryGetNumber(uint16& OutValue) const
+{
+	return LexTryParseString(OutValue, *Value);
+}
+
+template <typename CharType>
+bool TJsonValueNumberString<CharType>::TryGetNumber(uint32& OutValue) const
+{
+	return LexTryParseString(OutValue, *Value);
+}
+
+template <typename CharType>
+bool TJsonValueNumberString<CharType>::TryGetNumber(uint64& OutValue) const
+{
+	return LexTryParseString(OutValue, *Value);
+}
+
+template <typename CharType>
+bool TJsonValueNumberString<CharType>::TryGetBool(bool& OutBool) const
+{
+	OutBool = Value.ToBool();
+	return true;
+}
+
+template <typename CharType>
+bool TJsonValueNumberString<CharType>::PreferStringRepresentation() const
+{
+	return true;
+}
+
+template <typename CharType>
+SIZE_T TJsonValueNumberString<CharType>::GetMemoryFootprint() const
+{
+	return sizeof(*this) + GetAllocatedSize();
+}
+
 template <typename CharType>
 FString TJsonValueNumberString<CharType>::GetType() const
 {
 	return TJsonValueNumberType<CharType>::GetType();
 }
 
+template <typename CharType>
+SIZE_T TJsonValueNumberString<CharType>::GetAllocatedSize() const
+{
+	return Value.GetAllocatedSize();
+}
+
 template class JSON_API TJsonValueNumberString<TCHAR>;
 template class JSON_API TJsonValueNumberString<UTF8CHAR>;
 
+FJsonValueBoolean::FJsonValueBoolean(bool InBool)
+: Value(InBool) 
+{
+	Type = EJson::Boolean;
+}
+
+bool FJsonValueBoolean::TryGetNumber(double& OutNumber) const
+{
+	OutNumber = Value ? 1 : 0;
+	return true;
+}
+
+bool FJsonValueBoolean::TryGetBool(bool& OutBool) const
+{
+	OutBool = Value;
+	return true;
+}
+
+bool FJsonValueBoolean::TryGetString(FString& OutString) const
+{
+	OutString = Value ? TEXT("true") : TEXT("false");
+	return true;
+}
+
+bool FJsonValueBoolean::TryGetUtf8String(FUtf8String& OutString) const
+{
+	OutString = Value ? UTF8TEXT("true") : UTF8TEXT("false");
+	return true;
+}
+
+SIZE_T FJsonValueBoolean::GetMemoryFootprint() const
+{
+	return sizeof(*this);
+}
+
+FString FJsonValueBoolean::GetType() const
+{
+	return TEXT("Boolean");
+}
+
+FJsonValueArray::FJsonValueArray(const TArray<TSharedPtr<FJsonValue>>& InArray)
+: Value(InArray)
+{
+	Type = EJson::Array;
+}
+
+FJsonValueArray::FJsonValueArray(TArray<TSharedPtr<FJsonValue>>&& InArray)
+: Value(MoveTemp(InArray))
+{
+	Type = EJson::Array;
+}
+
+bool FJsonValueArray::TryGetArray(const TArray<TSharedPtr<FJsonValue>>*& OutArray) const
+{
+	OutArray = &Value;
+	return true;
+}
+
+bool FJsonValueArray::TryGetArray(TArray<TSharedPtr<FJsonValue>>*& OutArray)
+{
+	OutArray = &Value;
+	return true;
+}
+
+SIZE_T FJsonValueArray::GetMemoryFootprint() const
+{
+	 return sizeof(*this) + GetAllocatedSize(); 
+}
+
 SIZE_T FJsonValueArray::GetAllocatedSize() const
 {
 	SIZE_T SizeBytes = 0;
@@ -547,11 +849,58 @@ SIZE_T FJsonValueArray::GetAllocatedSize() const
 	return SizeBytes;
 }
 
+FString FJsonValueArray::GetType() const
+{
+	return TEXT("Array");
+}	
+
+FJsonValueObject::FJsonValueObject(TSharedPtr<FJsonObject> InObject)
+: Value(MoveTemp(InObject))
+{
+	Type = EJson::Object;
+}
+
+bool FJsonValueObject::TryGetObject(const TSharedPtr<FJsonObject>*& OutObject) const
+{
+	OutObject = &Value;
+	return true;
+}
+
+bool FJsonValueObject::TryGetObject(TSharedPtr<FJsonObject>*& OutObject)
+{
+	OutObject = &Value;
+	return true;
+}
+
+SIZE_T FJsonValueObject::GetMemoryFootprint() const
+{
+	return sizeof(*this) + GetAllocatedSize();
+}
+
 SIZE_T FJsonValueObject::GetAllocatedSize() const
 {
 	return Value.IsValid() ? Value->GetMemoryFootprint() : 0;
 }
 
+FString FJsonValueObject::GetType() const
+{
+	return TEXT("Object");
+}
+
+FJsonValueNull::FJsonValueNull()
+{
+	Type = EJson::Null;
+}
+
+SIZE_T FJsonValueNull::GetMemoryFootprint() const
+{
+	return sizeof(*this);
+}
+
+FString FJsonValueNull::GetType() const
+{
+	return TEXT("Null");
+}
 
 namespace UE::Json
 {
diff --git a/Engine/Source/Runtime/Json/Public/Dom/JsonValue.h b/Engine/Source/Runtime/Json/Public/Dom/JsonValue.h
--- a/Engine/Source/Runtime/Json/Public/Dom/JsonValue.h
+++ b/Engine/Source/Runtime/Json/Public/Dom/JsonValue.h
@@ -35,13 +35,13 @@ public:
 	JSON_API bool AsBool() const;
 
 	/** Returns this value as an array, logging an error and returning an empty array reference if not possible */
-	JSON_API const TArray< TSharedPtr<FJsonValue> >& AsArray() const;
+	JSON_API const TArray<TSharedPtr<FJsonValue>>& AsArray() const;
 
 	/** Returns this value as an object, throwing an error if this is not an Json Object */
 	JSON_API virtual const TSharedPtr<FJsonObject>& AsObject() const;
 
 	/** Tries to convert this value to a number, returning false if not possible */
-	virtual bool TryGetNumber(double& OutNumber) const { return false; }
+	JSON_API virtual bool TryGetNumber(double& OutNumber) const;
 
 	/** Tries to convert this value to a number, returning false if not possible */
 	JSON_API virtual bool TryGetNumber(float& OutNumber) const;
@@ -71,44 +71,65 @@ public:
 	JSON_API virtual bool TryGetNumber(uint64& OutNumber) const;
 
 	/** Tries to convert this value to a string, returning false if not possible */
-	virtual bool TryGetString(FString& OutString) const { return false; }
+	JSON_API virtual bool TryGetString(FString& OutString) const;
 
 	/** Tries to convert this value to a utf8 string, returning false if not possible */
-	virtual bool TryGetUtf8String(FUtf8String& OutString) const { return false; }
+	JSON_API virtual bool TryGetUtf8String(FUtf8String& OutString) const;
 
 	/** Tries to convert this value to a bool, returning false if not possible */
-	virtual bool TryGetBool(bool& OutBool) const { return false; }
+	JSON_API virtual bool TryGetBool(bool& OutBool) const;
 
 	/** Tries to convert this value to an array, returning false if not possible */
-	virtual bool TryGetArray(const TArray< TSharedPtr<FJsonValue> >*& OutArray) const { return false; }
+	JSON_API virtual bool TryGetArray(const TArray<TSharedPtr<FJsonValue>>*& OutArray) const;
 	
 	/** Tries to convert this value to an array, returning false if not possible */
-	virtual bool TryGetArray(TArray< TSharedPtr<FJsonValue> >*& OutArray) { return false; }
+	JSON_API virtual bool TryGetArray(TArray<TSharedPtr<FJsonValue>>*& OutArray);
 
 	/** Tries to convert this value to an object, returning false if not possible */
-	virtual bool TryGetObject(const TSharedPtr<FJsonObject>*& Object) const { return false; }
+	JSON_API virtual bool TryGetObject(const TSharedPtr<FJsonObject>*& Object) const;
 
 	/** Tries to convert this value to an object, returning false if not possible */
-	virtual bool TryGetObject(TSharedPtr<FJsonObject>*& Object) { return false; }
+	JSON_API virtual bool TryGetObject(TSharedPtr<FJsonObject>*& Object);
 
 	/** Returns whether or not a caller should prefer a string representation of the value, rather than the natural JSON type */
-	virtual bool PreferStringRepresentation() const { return false; }
+	JSON_API virtual bool PreferStringRepresentation() const;
 
 	/** Returns true if this value is a 'null' */
-	bool IsNull() const { return Type == EJson::Null || Type == EJson::None; }
+	bool IsNull() const
+	{
+		return Type == EJson::Null || Type == EJson::None;
+	}
 
 	/** Get a field of the same type as the argument */
-	void AsArgumentType(double                          & Value) { Value = AsNumber(); }
-	void AsArgumentType(FString                         & Value) { Value = AsString(); }
-	void AsArgumentType(bool                            & Value) { Value = AsBool  (); }
-	void AsArgumentType(TArray< TSharedPtr<FJsonValue> >& Value) { Value = AsArray (); }
-	void AsArgumentType(TSharedPtr<FJsonObject>         & Value) { Value = AsObject(); }
+	void AsArgumentType(double& Value)
+	{
+		Value = AsNumber();
+	}
+	
+	void AsArgumentType(FString& Value)
+	{
+		Value = AsString();
+	}
+	
+	void AsArgumentType(bool& Value)
+	{
+		Value = AsBool();
+	}
+	
+	void AsArgumentType(TArray<TSharedPtr<FJsonValue>>& Value)
+	{
+		Value = AsArray();
+	}
+	void AsArgumentType(TSharedPtr<FJsonObject>& Value)
+	{
+		Value = AsObject();
+	}
 
 	/**
 	 * Returns the memory footprint for this object in Bytes, including sizeof(*this) and allocated memory.
 	 * All children should implement this so their memory layout is properly accounted for
 	 */
-	JSON_API virtual SIZE_T GetMemoryFootprint() const { return sizeof(*this); }
+	JSON_API virtual SIZE_T GetMemoryFootprint() const;
 
 	EJson Type;
 
@@ -120,8 +141,8 @@ public:
 protected:
 	UE_NONCOPYABLE(FJsonValue)
 
-	FJsonValue() : Type(EJson::None) {}
-	virtual ~FJsonValue() {}
+	JSON_API FJsonValue();
+	JSON_API virtual ~FJsonValue();
 
 	virtual FString GetType() const = 0;
 
@@ -149,22 +170,19 @@ public:
 	virtual bool TryGetString(FString& OutString) const override;	
 	virtual bool TryGetUtf8String(FUtf8String& OutString) const override;
 	virtual bool TryGetNumber(double& OutDouble) const override;
-
-	virtual bool TryGetNumber(int32& OutValue) const override	{ LexFromString(OutValue, *Value); return true; }
-	
-	virtual bool TryGetNumber(uint32& OutValue) const override	{ LexFromString(OutValue, *Value); return true; }
-
-	virtual bool TryGetNumber(int64& OutValue) const override	{ LexFromString(OutValue, *Value); return true; }
-
-	virtual bool TryGetNumber(uint64& OutValue) const override	{ LexFromString(OutValue, *Value); return true; }
-
-	virtual bool TryGetBool(bool& OutBool) const override		{ OutBool = Value.ToBool(); return true; }
+	virtual bool TryGetNumber(int32& OutValue) const override;	
+	virtual bool TryGetNumber(uint32& OutValue) const override;
+	virtual bool TryGetNumber(int64& OutValue) const override;
+	virtual bool TryGetNumber(uint64& OutValue) const override;
+	virtual bool TryGetBool(bool& OutBool) const override;
 
 	// Way to check if string value is empty without copying the string 
-	bool IsEmpty() const { return Value.IsEmpty(); }
-
-	virtual SIZE_T GetMemoryFootprint() const override { return sizeof(*this) + GetAllocatedSize(); }
+	bool IsEmpty() const
+	{
+		return Value.IsEmpty();
+	}
 
+	virtual SIZE_T GetMemoryFootprint() const override;
 
 protected:
 	TString<CharType> Value;
@@ -172,7 +190,7 @@ protected:
 	virtual FString GetType() const override;
 
 	/** Helper to calculate allocated size of the Value string */
-	SIZE_T GetAllocatedSize() const { return Value.GetAllocatedSize(); }
+	SIZE_T GetAllocatedSize() const;
 };
 
 #if PLATFORM_COMPILER_CLANG
@@ -190,18 +208,17 @@ using FJsonValueString = TJsonValueString<TCHAR>;
 class FJsonValueNumber : public FJsonValue
 {
 public:
-	FJsonValueNumber(double InNumber) : Value(InNumber) {Type = EJson::Number;}
-	virtual bool TryGetNumber(double& OutNumber) const override		{ OutNumber = Value; return true; }
-	virtual bool TryGetBool(bool& OutBool) const override			{ OutBool = (Value != 0.0); return true; }
-	virtual bool TryGetString(FString& OutString) const override	{ OutString = FString::SanitizeFloat(Value, 0); return true; }
-	virtual bool TryGetUtf8String(FUtf8String& OutString) const override { OutString = FUtf8String::SanitizeFloat(Value, 0); return true; }
-	virtual SIZE_T GetMemoryFootprint() const override { return sizeof(*this); }
+	JSON_API FJsonValueNumber(double InNumber);
+	JSON_API virtual bool TryGetNumber(double& OutNumber) const override;
+	JSON_API virtual bool TryGetBool(bool& OutBool) const override;
+	JSON_API virtual bool TryGetString(FString& OutString) const override;
+	JSON_API virtual bool TryGetUtf8String(FUtf8String& OutString) const override;
+	JSON_API virtual SIZE_T GetMemoryFootprint() const override;
 
 protected:
-
 	double Value;
 
-	virtual FString GetType() const override {return TEXT("Number");}
+	JSON_API virtual FString GetType() const override;
 };
 
 /** A Json Number Value, stored internally as a string so as not to lose precision */
@@ -214,21 +231,19 @@ public:
 
 	virtual bool TryGetString(FString& OutString) const override;	
 	virtual bool TryGetUtf8String(FUtf8String& OutString) const override;
-	
-	virtual bool TryGetNumber(double& OutDouble) const override { return LexTryParseString(OutDouble, *Value); }
-	virtual bool TryGetNumber(float &OutDouble) const override { return LexTryParseString(OutDouble, *Value); }
-	virtual bool TryGetNumber(int8& OutValue) const override { return LexTryParseString(OutValue, *Value); }
-	virtual bool TryGetNumber(int16& OutValue) const override { return LexTryParseString(OutValue, *Value); }
-	virtual bool TryGetNumber(int32& OutValue) const override { return LexTryParseString(OutValue, *Value); }
-	virtual bool TryGetNumber(int64& OutValue) const override { return LexTryParseString(OutValue, *Value); }
-	virtual bool TryGetNumber(uint8& OutValue) const override { return LexTryParseString(OutValue, *Value); }
-	virtual bool TryGetNumber(uint16& OutValue) const override { return LexTryParseString(OutValue, *Value); }
-	virtual bool TryGetNumber(uint32& OutValue) const override { return LexTryParseString(OutValue, *Value); }
-	virtual bool TryGetNumber(uint64& OutValue) const override { return LexTryParseString(OutValue, *Value); }
-	virtual bool TryGetBool(bool& OutBool) const override { OutBool = Value.ToBool(); return true; }
-	virtual bool PreferStringRepresentation() const override { return true; }
-	virtual SIZE_T GetMemoryFootprint() const override { return sizeof(*this) + GetAllocatedSize(); }
-
+	virtual bool TryGetNumber(double& OutDouble) const override;
+	virtual bool TryGetNumber(float &OutDouble) const override;
+	virtual bool TryGetNumber(int8& OutValue) const override;
+	virtual bool TryGetNumber(int16& OutValue) const override;
+	virtual bool TryGetNumber(int32& OutValue) const override;
+	virtual bool TryGetNumber(int64& OutValue) const override;
+	virtual bool TryGetNumber(uint8& OutValue) const override;
+	virtual bool TryGetNumber(uint16& OutValue) const override;
+	virtual bool TryGetNumber(uint32& OutValue) const override;
+	virtual bool TryGetNumber(uint64& OutValue) const override;
+	virtual bool TryGetBool(bool& OutBool) const override;
+	virtual bool PreferStringRepresentation() const override;
+	virtual SIZE_T GetMemoryFootprint() const override;
 
 protected:
 	TString<CharType> Value;
@@ -236,7 +251,7 @@ protected:
 	virtual FString GetType() const override;
 
 	/** Helper to calculate allocated size of the Value string */
-	SIZE_T GetAllocatedSize() const { return Value.GetAllocatedSize(); }
+	SIZE_T GetAllocatedSize() const;
 };
 
 extern template class UE_EXTERN_TEMPLATE_API TJsonValueNumberString<TCHAR>;
@@ -250,34 +265,34 @@ using FJsonValueNumberString = TJsonValueNumberString<TCHAR>;
 class FJsonValueBoolean : public FJsonValue
 {
 public:
-	FJsonValueBoolean(bool InBool) : Value(InBool) {Type = EJson::Boolean;}
-	virtual bool TryGetNumber(double& OutNumber) const override		{ OutNumber = Value ? 1 : 0; return true; }
-	virtual bool TryGetBool(bool& OutBool) const override			{ OutBool = Value; return true; }
-	virtual bool TryGetString(FString& OutString) const override	{ OutString = Value ? TEXT("true") : TEXT("false"); return true; }
-	virtual bool TryGetUtf8String(FUtf8String& OutString) const override { OutString = Value ? UTF8TEXT("true") : UTF8TEXT("false"); return true; }
-	virtual SIZE_T GetMemoryFootprint() const override { return sizeof(*this); }
+	JSON_API FJsonValueBoolean(bool InBool);
+	JSON_API virtual bool TryGetNumber(double& OutNumber) const override;
+	JSON_API virtual bool TryGetBool(bool& OutBool) const override;
+	JSON_API virtual bool TryGetString(FString& OutString) const override;
+	JSON_API virtual bool TryGetUtf8String(FUtf8String& OutString) const override;
+	JSON_API virtual SIZE_T GetMemoryFootprint() const override;
 
 protected:
 	bool Value;
 
-	virtual FString GetType() const override {return TEXT("Boolean");}
+	JSON_API virtual FString GetType() const override;
 };
 
 /** A Json Array Value. */
 class FJsonValueArray : public FJsonValue
 {
 public:
-	FJsonValueArray(const TArray< TSharedPtr<FJsonValue> >& InArray) : Value(InArray) {Type = EJson::Array;}
-	FJsonValueArray(TArray< TSharedPtr<FJsonValue> >&& InArray) : Value(MoveTemp(InArray)) {Type = EJson::Array;}
-	virtual bool TryGetArray(const TArray< TSharedPtr<FJsonValue> >*& OutArray) const override	{ OutArray = &Value; return true; }
-	virtual bool TryGetArray(TArray< TSharedPtr<FJsonValue> >*& OutArray) override				{ OutArray = &Value; return true; }
-	virtual SIZE_T GetMemoryFootprint() const override { return sizeof(*this) + GetAllocatedSize(); }
+	JSON_API FJsonValueArray(const TArray<TSharedPtr<FJsonValue>>& InArray);
+	JSON_API FJsonValueArray(TArray<TSharedPtr<FJsonValue>>&& InArray);
+	JSON_API virtual bool TryGetArray(const TArray<TSharedPtr<FJsonValue>>*& OutArray) const override;
+	JSON_API virtual bool TryGetArray(TArray<TSharedPtr<FJsonValue>>*& OutArray) override;
+	JSON_API virtual SIZE_T GetMemoryFootprint() const override;
 
 protected:
 	TArray<TSharedPtr<FJsonValue>> Value;
 
-	virtual FString GetType() const override {return TEXT("Array");}	
-	/** Helper to calculate allocated size of the Value array and its contents */
+	JSON_API virtual FString GetType() const override;	
+/** Helper to calculate allocated size of the Value array and its contents */
 	JSON_API SIZE_T GetAllocatedSize() const;
 };
 
@@ -286,14 +301,14 @@ protected:
 class FJsonValueObject : public FJsonValue
 {
 public:
-	FJsonValueObject(TSharedPtr<FJsonObject> InObject) : Value(MoveTemp(InObject)) {Type = EJson::Object;}
-	virtual bool TryGetObject(const TSharedPtr<FJsonObject>*& OutObject) const override	{ OutObject = &Value; return true; }
-	virtual bool TryGetObject(TSharedPtr<FJsonObject>*& OutObject) override				{ OutObject = &Value; return true; }
-	virtual SIZE_T GetMemoryFootprint() const override { return sizeof(*this) + GetAllocatedSize(); }
+	JSON_API FJsonValueObject(TSharedPtr<FJsonObject> InObject);
+	JSON_API virtual bool TryGetObject(const TSharedPtr<FJsonObject>*& OutObject) const override;
+	JSON_API virtual bool TryGetObject(TSharedPtr<FJsonObject>*& OutObject) override;
+	JSON_API virtual SIZE_T GetMemoryFootprint() const override;
 protected:
 	TSharedPtr<FJsonObject> Value;
 
-	virtual FString GetType() const override {return TEXT("Object");}
+	JSON_API virtual FString GetType() const override;
 	/** Helper to calculate allocated size of the Value object and its contents */
 	JSON_API SIZE_T GetAllocatedSize() const;
 };
@@ -303,11 +318,11 @@ protected:
 class FJsonValueNull : public FJsonValue
 {
 public:
-	FJsonValueNull() {Type = EJson::Null;}
-	virtual SIZE_T GetMemoryFootprint() const override { return sizeof(*this); }
+	JSON_API FJsonValueNull();
+	JSON_API virtual SIZE_T GetMemoryFootprint() const override;
 
 protected:
-	virtual FString GetType() const override {return TEXT("Null");}
+	JSON_API virtual FString GetType() const override;
 };
 
 namespace UE::Json
