Fix various build-time issues for modular build

#rnx

diff --git a/Engine/Source/Runtime/Engine/Classes/Sound/SoundAttenuation.h b/Engine/Source/Runtime/Engine/Classes/Sound/SoundAttenuation.h
--- a/Engine/Source/Runtime/Engine/Classes/Sound/SoundAttenuation.h
+++ b/Engine/Source/Runtime/Engine/Classes/Sound/SoundAttenuation.h
@@ -408,6 +408,7 @@ struct FSoundAttenuationSettings : public FBaseAttenuationSettings
 
 	ENGINE_API FSoundAttenuationSettings();
 	ENGINE_API ~FSoundAttenuationSettings();
+	ENGINE_API FSoundAttenuationSettings(const FSoundAttenuationSettings&);
 
 	ENGINE_API bool operator==(const FSoundAttenuationSettings& Other) const;
 #if WITH_EDITORONLY_DATA
@@ -419,7 +420,7 @@ struct FSoundAttenuationSettings : public FBaseAttenuationSettings
 	ENGINE_API float GetFocusAttenuation(const struct FGlobalFocusSettings& FocusSettings, float FocusFactor) const;
 	ENGINE_API float GetFocusDistanceScale(const struct FGlobalFocusSettings& FocusSettings, float FocusFactor) const;
 
-	void AddStructReferencedObjects(FReferenceCollector& Collector);
+	ENGINE_API void AddStructReferencedObjects(FReferenceCollector& Collector);
 };
 
 template<>
diff --git a/Engine/Source/Runtime/Engine/Private/SoundAttenuation.cpp b/Engine/Source/Runtime/Engine/Private/SoundAttenuation.cpp
--- a/Engine/Source/Runtime/Engine/Private/SoundAttenuation.cpp
+++ b/Engine/Source/Runtime/Engine/Private/SoundAttenuation.cpp
@@ -90,6 +90,8 @@ FSoundAttenuationSettings::FSoundAttenuationSettings()
 
 FSoundAttenuationSettings::~FSoundAttenuationSettings() = default;
 
+FSoundAttenuationSettings::FSoundAttenuationSettings(const FSoundAttenuationSettings&) = default;
+
 #if WITH_EDITORONLY_DATA
 void FSoundAttenuationSettings::PostSerialize(const FArchive& Ar)
 {
diff --git a/Engine/Source/Runtime/Json/Private/Dom/JsonValue.cpp b/Engine/Source/Runtime/Json/Private/Dom/JsonValue.cpp
--- a/Engine/Source/Runtime/Json/Private/Dom/JsonValue.cpp
+++ b/Engine/Source/Runtime/Json/Private/Dom/JsonValue.cpp
@@ -8,9 +8,7 @@ FJsonValue::FJsonValue()
 {
 }
 
-FJsonValue::~FJsonValue()
-{
-}
+FJsonValue::~FJsonValue() = default;
 
 double FJsonValue::AsNumber() const
 {
@@ -425,153 +423,14 @@ void FJsonValue::ErrorMessage(const FString& InType) const
 	}
 }
 
-
-namespace 
-{
-	template <typename CharType>
-	struct TJsonValueStringType
-	{
-		static_assert(sizeof(CharType) == 0, "Unsupported type");
-	};
-
-	template <>
-	struct TJsonValueStringType<TCHAR>
-	{
-		static FString GetType()
-		{
-			return TEXT("String");
-		}
-	};
-
-	template <>
-	struct TJsonValueStringType<UTF8CHAR>
-	{
-		static FString GetType()
-		{
-			return TEXT("Utf8String");
-		}
-	};
-}
-
-template <typename CharType>
-TJsonValueString<CharType>::TJsonValueString(const TString<CharType>& InString)
-: Value(InString)
-{
-	Type = EJson::String;
-}
-
-template <typename CharType>
-TJsonValueString<CharType>::TJsonValueString(TString<CharType>&& InString)
-: Value(MoveTemp(InString))
-{
-	Type = EJson::String;
-}
-
-template <typename CharType>
-bool TJsonValueString<CharType>::TryGetString(FString& OutString) const
-{
-	if constexpr (std::is_same_v<FString, TString<CharType>>)
-	{
-		OutString = Value;
-	}
-	else
-	{
-		OutString = FString(Value);
-	}
-	return true;
-}
-
-template <typename CharType>
-bool TJsonValueString<CharType>::TryGetUtf8String(FUtf8String& OutString) const
-{
-	if constexpr (std::is_same_v<FUtf8String, TString<CharType>>)
-	{
-		OutString = Value;
-	}
-	else
-	{
-		OutString = FUtf8String(Value);
-	}
-	return true;
-}
-
-template <typename CharType>
-bool TJsonValueString<CharType>::TryGetNumber(double& OutDouble) const
-{
-	if (Value.IsNumeric())
-	{
-		OutDouble = TCString<CharType>::Atod(*Value);
-		return true;
-	}
-	else
-	{
-		return false;
-	}
-}
-
-template <typename CharType>
-bool TJsonValueString<CharType>::TryGetNumber(int32& OutValue) const
-{
-	LexFromString(OutValue, *Value);
-	return true;
-}	
-
-template <typename CharType>
-bool TJsonValueString<CharType>::TryGetNumber(uint32& OutValue) const
-{
-	LexFromString(OutValue, *Value);
-	return true;
-}
-
-template <typename CharType>
-bool TJsonValueString<CharType>::TryGetNumber(int64& OutValue) const
-{
-	LexFromString(OutValue, *Value);
-	return true; 
-}
-
-template <typename CharType>
-bool TJsonValueString<CharType>::TryGetNumber(uint64& OutValue) const
-{
-	LexFromString(OutValue, *Value);
-	return true;
-}
-
-template <typename CharType>
-bool TJsonValueString<CharType>::TryGetBool(bool& OutBool) const
-{
-	OutBool = Value.ToBool();
-	return true;
-}
-
-template <typename CharType>
-SIZE_T TJsonValueString<CharType>::GetMemoryFootprint() const
-{
-	return sizeof(*this) + GetAllocatedSize();
-}
-
-template <typename CharType>
-FString TJsonValueString<CharType>::GetType() const
-{
-	return TJsonValueStringType<CharType>::GetType();
-}
-
-template <typename CharType>
-SIZE_T TJsonValueString<CharType>::GetAllocatedSize() const
-{
-	return Value.GetAllocatedSize();
-}
-
-template class JSON_API TJsonValueString<TCHAR>;
-template class JSON_API TJsonValueString<UTF8CHAR>;
-
-
 FJsonValueNumber::FJsonValueNumber(double InNumber)
 : Value(InNumber)
 {
 	Type = EJson::Number;
 }
 
+FJsonValueNumber::~FJsonValueNumber() = default;
+
 bool FJsonValueNumber::TryGetNumber(double& OutNumber) const
 {
 	OutNumber = Value;
@@ -606,175 +465,14 @@ FString FJsonValueNumber::GetType() const
 	return TEXT("Number");
 }
 
-namespace
-{
-template <typename CharType>
-struct TJsonValueNumberType
-{
-	static_assert(sizeof(CharType), "Unsupported type");
-};
-
-template <>
-struct TJsonValueNumberType<TCHAR>
-{
-	static FString GetType()
-	{
-		return TEXT("NumberString");
-	}
-};
-
-template <>
-struct TJsonValueNumberType<UTF8CHAR>
-{
-	static FString GetType()
-	{
-		return TEXT("Utf8NumberString");
-	}
-};
-}
-
-template <typename CharType>
-TJsonValueNumberString<CharType>::TJsonValueNumberString(const TString<CharType>& InString)
-: Value(InString)
-{
-	Type = EJson::Number;
-}
-
-template <typename CharType>
-TJsonValueNumberString<CharType>::TJsonValueNumberString(TString<CharType>&& InString)
-: Value(MoveTemp(InString))
-{
-	Type = EJson::Number;
-}
-
-template <typename CharType>
-bool TJsonValueNumberString<CharType>::TryGetString(FString& OutString) const
-{
-	if constexpr (std::is_same_v<FString, TString<CharType>>)
-	{
-		OutString = Value;
-	}
-	else
-	{
-		OutString = FString(Value);
-	}
-	return true;
-}
-
-template <typename CharType>
-bool TJsonValueNumberString<CharType>::TryGetUtf8String(FUtf8String& OutString) const
-{
-	if constexpr (std::is_same_v<FUtf8String, TString<CharType>>)
-	{
-		OutString = Value;
-	}
-	else
-	{
-		OutString = FUtf8String(Value);
-	}
-	return true;
-}
-
-template <typename CharType>
-bool TJsonValueNumberString<CharType>::TryGetNumber(double& OutDouble) const
-{
-	return LexTryParseString(OutDouble, *Value);
-}
-
-template <typename CharType>
-bool TJsonValueNumberString<CharType>::TryGetNumber(float &OutDouble) const
-{
-	return LexTryParseString(OutDouble, *Value);
-}
-
-template <typename CharType>
-bool TJsonValueNumberString<CharType>::TryGetNumber(int8& OutValue) const
-{
-	return LexTryParseString(OutValue, *Value);
-}
-
-template <typename CharType>
-bool TJsonValueNumberString<CharType>::TryGetNumber(int16& OutValue) const
-{
-	return LexTryParseString(OutValue, *Value);
-}
-
-template <typename CharType>
-bool TJsonValueNumberString<CharType>::TryGetNumber(int32& OutValue) const
-{
-	return LexTryParseString(OutValue, *Value);
-}
-
-template <typename CharType>
-bool TJsonValueNumberString<CharType>::TryGetNumber(int64& OutValue) const
-{
-	return LexTryParseString(OutValue, *Value);
-}
-
-template <typename CharType>
-bool TJsonValueNumberString<CharType>::TryGetNumber(uint8& OutValue) const
-{
-	return LexTryParseString(OutValue, *Value);
-}
-
-template <typename CharType>
-bool TJsonValueNumberString<CharType>::TryGetNumber(uint16& OutValue) const
-{
-	return LexTryParseString(OutValue, *Value);
-}
-
-template <typename CharType>
-bool TJsonValueNumberString<CharType>::TryGetNumber(uint32& OutValue) const
-{
-	return LexTryParseString(OutValue, *Value);
-}
-
-template <typename CharType>
-bool TJsonValueNumberString<CharType>::TryGetNumber(uint64& OutValue) const
-{
-	return LexTryParseString(OutValue, *Value);
-}
-
-template <typename CharType>
-bool TJsonValueNumberString<CharType>::TryGetBool(bool& OutBool) const
-{
-	OutBool = Value.ToBool();
-	return true;
-}
-
-template <typename CharType>
-bool TJsonValueNumberString<CharType>::PreferStringRepresentation() const
-{
-	return true;
-}
-
-template <typename CharType>
-SIZE_T TJsonValueNumberString<CharType>::GetMemoryFootprint() const
-{
-	return sizeof(*this) + GetAllocatedSize();
-}
-
-template <typename CharType>
-FString TJsonValueNumberString<CharType>::GetType() const
-{
-	return TJsonValueNumberType<CharType>::GetType();
-}
-
-template <typename CharType>
-SIZE_T TJsonValueNumberString<CharType>::GetAllocatedSize() const
-{
-	return Value.GetAllocatedSize();
-}
-
-template class JSON_API TJsonValueNumberString<TCHAR>;
-template class JSON_API TJsonValueNumberString<UTF8CHAR>;
-
 FJsonValueBoolean::FJsonValueBoolean(bool InBool)
 : Value(InBool) 
 {
 	Type = EJson::Boolean;
 }
 
+FJsonValueBoolean::~FJsonValueBoolean() = default;
+
 bool FJsonValueBoolean::TryGetNumber(double& OutNumber) const
 {
 	OutNumber = Value ? 1 : 0;
@@ -821,6 +519,8 @@ FJsonValueArray::FJsonValueArray(TArray<TSharedPtr<FJsonValue>>&& InArray)
 	Type = EJson::Array;
 }
 
+FJsonValueArray::~FJsonValueArray() = default;
+
 bool FJsonValueArray::TryGetArray(const TArray<TSharedPtr<FJsonValue>>*& OutArray) const
 {
 	OutArray = &Value;
@@ -860,6 +560,8 @@ FJsonValueObject::FJsonValueObject(TSharedPtr<FJsonObject> InObject)
 	Type = EJson::Object;
 }
 
+FJsonValueObject::~FJsonValueObject() = default;
+
 bool FJsonValueObject::TryGetObject(const TSharedPtr<FJsonObject>*& OutObject) const
 {
 	OutObject = &Value;
@@ -892,6 +594,8 @@ FJsonValueNull::FJsonValueNull()
 	Type = EJson::Null;
 }
 
+FJsonValueNull::~FJsonValueNull() = default;
+
 SIZE_T FJsonValueNull::GetMemoryFootprint() const
 {
 	return sizeof(*this);
diff --git a/Engine/Source/Runtime/Json/Public/Dom/JsonValue.h b/Engine/Source/Runtime/Json/Public/Dom/JsonValue.h
--- a/Engine/Source/Runtime/Json/Public/Dom/JsonValue.h
+++ b/Engine/Source/Runtime/Json/Public/Dom/JsonValue.h
@@ -164,44 +164,35 @@ template <typename CharType>
 class TJsonValueString : public FJsonValue
 {
 public:
-	TJsonValueString(const TString<CharType>& InString);
-	TJsonValueString(TString<CharType>&& InString);
-
-	virtual bool TryGetString(FString& OutString) const override;	
-	virtual bool TryGetUtf8String(FUtf8String& OutString) const override;
-	virtual bool TryGetNumber(double& OutDouble) const override;
-	virtual bool TryGetNumber(int32& OutValue) const override;	
-	virtual bool TryGetNumber(uint32& OutValue) const override;
-	virtual bool TryGetNumber(int64& OutValue) const override;
-	virtual bool TryGetNumber(uint64& OutValue) const override;
-	virtual bool TryGetBool(bool& OutBool) const override;
+	FORCEINLINE TJsonValueString(const TString<CharType>& InString);
+	FORCEINLINE TJsonValueString(TString<CharType>&& InString);
+
+	FORCEINLINE virtual bool TryGetString(FString& OutString) const override;	
+	FORCEINLINE virtual bool TryGetUtf8String(FUtf8String& OutString) const override;
+	FORCEINLINE virtual bool TryGetNumber(double& OutDouble) const override;
+	FORCEINLINE virtual bool TryGetNumber(int32& OutValue) const override;	
+	FORCEINLINE virtual bool TryGetNumber(uint32& OutValue) const override;
+	FORCEINLINE virtual bool TryGetNumber(int64& OutValue) const override;
+	FORCEINLINE virtual bool TryGetNumber(uint64& OutValue) const override;
+	FORCEINLINE virtual bool TryGetBool(bool& OutBool) const override;
 
 	// Way to check if string value is empty without copying the string 
-	bool IsEmpty() const
+	FORCEINLINE bool IsEmpty() const
 	{
 		return Value.IsEmpty();
 	}
 
-	virtual SIZE_T GetMemoryFootprint() const override;
+	FORCEINLINE virtual SIZE_T GetMemoryFootprint() const override;
 
 protected:
 	TString<CharType> Value;
 
-	virtual FString GetType() const override;
+	FORCEINLINE virtual FString GetType() const override;
 
 	/** Helper to calculate allocated size of the Value string */
-	SIZE_T GetAllocatedSize() const;
+	FORCEINLINE SIZE_T GetAllocatedSize() const;
 };
 
-#if PLATFORM_COMPILER_CLANG
-#define UE_EXTERN_TEMPLATE_API JSON_API
-#else
-#define UE_EXTERN_TEMPLATE_API
-#endif
-
-extern template class UE_EXTERN_TEMPLATE_API TJsonValueString<TCHAR>;
-extern template class UE_EXTERN_TEMPLATE_API TJsonValueString<UTF8CHAR>;
-
 using FJsonValueString = TJsonValueString<TCHAR>;
 
 //** A Json Number Value. */
@@ -209,6 +200,8 @@ class FJsonValueNumber : public FJsonValue
 {
 public:
 	JSON_API FJsonValueNumber(double InNumber);
+	JSON_API virtual ~FJsonValueNumber();
+;
 	JSON_API virtual bool TryGetNumber(double& OutNumber) const override;
 	JSON_API virtual bool TryGetBool(bool& OutBool) const override;
 	JSON_API virtual bool TryGetString(FString& OutString) const override;
@@ -226,39 +219,34 @@ template <typename CharType>
 class TJsonValueNumberString : public FJsonValue
 {
 public:
-	TJsonValueNumberString(const TString<CharType>& InString);
-	TJsonValueNumberString(TString<CharType>&& InString);
-
-	virtual bool TryGetString(FString& OutString) const override;	
-	virtual bool TryGetUtf8String(FUtf8String& OutString) const override;
-	virtual bool TryGetNumber(double& OutDouble) const override;
-	virtual bool TryGetNumber(float &OutDouble) const override;
-	virtual bool TryGetNumber(int8& OutValue) const override;
-	virtual bool TryGetNumber(int16& OutValue) const override;
-	virtual bool TryGetNumber(int32& OutValue) const override;
-	virtual bool TryGetNumber(int64& OutValue) const override;
-	virtual bool TryGetNumber(uint8& OutValue) const override;
-	virtual bool TryGetNumber(uint16& OutValue) const override;
-	virtual bool TryGetNumber(uint32& OutValue) const override;
-	virtual bool TryGetNumber(uint64& OutValue) const override;
-	virtual bool TryGetBool(bool& OutBool) const override;
-	virtual bool PreferStringRepresentation() const override;
-	virtual SIZE_T GetMemoryFootprint() const override;
+	FORCEINLINE TJsonValueNumberString(const TString<CharType>& InString);
+	FORCEINLINE TJsonValueNumberString(TString<CharType>&& InString);
+
+	FORCEINLINE virtual bool TryGetString(FString& OutString) const override;	
+	FORCEINLINE virtual bool TryGetUtf8String(FUtf8String& OutString) const override;
+	FORCEINLINE virtual bool TryGetNumber(double& OutDouble) const override;
+	FORCEINLINE virtual bool TryGetNumber(float &OutDouble) const override;
+	FORCEINLINE virtual bool TryGetNumber(int8& OutValue) const override;
+	FORCEINLINE virtual bool TryGetNumber(int16& OutValue) const override;
+	FORCEINLINE virtual bool TryGetNumber(int32& OutValue) const override;
+	FORCEINLINE virtual bool TryGetNumber(int64& OutValue) const override;
+	FORCEINLINE virtual bool TryGetNumber(uint8& OutValue) const override;
+	FORCEINLINE virtual bool TryGetNumber(uint16& OutValue) const override;
+	FORCEINLINE virtual bool TryGetNumber(uint32& OutValue) const override;
+	FORCEINLINE virtual bool TryGetNumber(uint64& OutValue) const override;
+	FORCEINLINE virtual bool TryGetBool(bool& OutBool) const override;
+	FORCEINLINE virtual bool PreferStringRepresentation() const override;
+	FORCEINLINE virtual SIZE_T GetMemoryFootprint() const override;
 
 protected:
 	TString<CharType> Value;
 
-	virtual FString GetType() const override;
+	FORCEINLINE virtual FString GetType() const override;
 
 	/** Helper to calculate allocated size of the Value string */
-	SIZE_T GetAllocatedSize() const;
+	FORCEINLINE SIZE_T GetAllocatedSize() const;
 };
 
-extern template class UE_EXTERN_TEMPLATE_API TJsonValueNumberString<TCHAR>;
-extern template class UE_EXTERN_TEMPLATE_API TJsonValueNumberString<UTF8CHAR>;
-
-#undef UE_EXTERN_TEMPLATE_API
-
 using FJsonValueNumberString = TJsonValueNumberString<TCHAR>;
 
 /** A Json Boolean Value. */
@@ -266,6 +254,8 @@ class FJsonValueBoolean : public FJsonValue
 {
 public:
 	JSON_API FJsonValueBoolean(bool InBool);
+	JSON_API virtual ~FJsonValueBoolean();
+
 	JSON_API virtual bool TryGetNumber(double& OutNumber) const override;
 	JSON_API virtual bool TryGetBool(bool& OutBool) const override;
 	JSON_API virtual bool TryGetString(FString& OutString) const override;
@@ -284,6 +274,8 @@ class FJsonValueArray : public FJsonValue
 public:
 	JSON_API FJsonValueArray(const TArray<TSharedPtr<FJsonValue>>& InArray);
 	JSON_API FJsonValueArray(TArray<TSharedPtr<FJsonValue>>&& InArray);
+	JSON_API virtual ~FJsonValueArray();
+
 	JSON_API virtual bool TryGetArray(const TArray<TSharedPtr<FJsonValue>>*& OutArray) const override;
 	JSON_API virtual bool TryGetArray(TArray<TSharedPtr<FJsonValue>>*& OutArray) override;
 	JSON_API virtual SIZE_T GetMemoryFootprint() const override;
@@ -302,6 +294,8 @@ class FJsonValueObject : public FJsonValue
 {
 public:
 	JSON_API FJsonValueObject(TSharedPtr<FJsonObject> InObject);
+	JSON_API virtual ~FJsonValueObject();
+
 	JSON_API virtual bool TryGetObject(const TSharedPtr<FJsonObject>*& OutObject) const override;
 	JSON_API virtual bool TryGetObject(TSharedPtr<FJsonObject>*& OutObject) override;
 	JSON_API virtual SIZE_T GetMemoryFootprint() const override;
@@ -319,6 +313,8 @@ class FJsonValueNull : public FJsonValue
 {
 public:
 	JSON_API FJsonValueNull();
+	JSON_API virtual ~FJsonValueNull();
+
 	JSON_API virtual SIZE_T GetMemoryFootprint() const override;
 
 protected:
@@ -381,3 +377,299 @@ inline bool operator!=(const JsonSimpleValueVariant& Lhs, const JsonSimpleValueV
 	return !(Lhs == Rhs);
 }
 #endif
+
+namespace
+{
+	template <typename CharType>
+	struct TJsonValueStringType
+	{
+		static_assert(sizeof(CharType) == 0, "Unsupported type");
+	};
+
+	template <>
+	struct TJsonValueStringType<TCHAR>
+	{
+		static FString GetType()
+		{
+			return TEXT("String");
+		}
+	};
+
+	template <>
+	struct TJsonValueStringType<UTF8CHAR>
+	{
+		static FString GetType()
+		{
+			return TEXT("Utf8String");
+		}
+	};
+}
+
+template <typename CharType>
+TJsonValueString<CharType>::TJsonValueString(const TString<CharType>& InString)
+	: Value(InString)
+{
+	Type = EJson::String;
+}
+
+template <typename CharType>
+TJsonValueString<CharType>::TJsonValueString(TString<CharType>&& InString)
+	: Value(MoveTemp(InString))
+{
+	Type = EJson::String;
+}
+
+template <typename CharType>
+bool TJsonValueString<CharType>::TryGetString(FString& OutString) const
+{
+	if constexpr (std::is_same_v<FString, TString<CharType>>)
+	{
+		OutString = Value;
+	}
+	else
+	{
+		OutString = FString(Value);
+	}
+	return true;
+}
+
+template <typename CharType>
+bool TJsonValueString<CharType>::TryGetUtf8String(FUtf8String& OutString) const
+{
+	if constexpr (std::is_same_v<FUtf8String, TString<CharType>>)
+	{
+		OutString = Value;
+	}
+	else
+	{
+		OutString = FUtf8String(Value);
+	}
+	return true;
+}
+
+template <typename CharType>
+bool TJsonValueString<CharType>::TryGetNumber(double& OutDouble) const
+{
+	if (Value.IsNumeric())
+	{
+		OutDouble = TCString<CharType>::Atod(*Value);
+		return true;
+	}
+	else
+	{
+		return false;
+	}
+}
+
+template <typename CharType>
+bool TJsonValueString<CharType>::TryGetNumber(int32& OutValue) const
+{
+	LexFromString(OutValue, *Value);
+	return true;
+}
+
+template <typename CharType>
+bool TJsonValueString<CharType>::TryGetNumber(uint32& OutValue) const
+{
+	LexFromString(OutValue, *Value);
+	return true;
+}
+
+template <typename CharType>
+bool TJsonValueString<CharType>::TryGetNumber(int64& OutValue) const
+{
+	LexFromString(OutValue, *Value);
+	return true;
+}
+
+template <typename CharType>
+bool TJsonValueString<CharType>::TryGetNumber(uint64& OutValue) const
+{
+	LexFromString(OutValue, *Value);
+	return true;
+}
+
+template <typename CharType>
+bool TJsonValueString<CharType>::TryGetBool(bool& OutBool) const
+{
+	OutBool = Value.ToBool();
+	return true;
+}
+
+template <typename CharType>
+SIZE_T TJsonValueString<CharType>::GetMemoryFootprint() const
+{
+	return sizeof(*this) + GetAllocatedSize();
+}
+
+template <typename CharType>
+FString TJsonValueString<CharType>::GetType() const
+{
+	return TJsonValueStringType<CharType>::GetType();
+}
+
+template <typename CharType>
+SIZE_T TJsonValueString<CharType>::GetAllocatedSize() const
+{
+	return Value.GetAllocatedSize();
+}
+
+namespace
+{
+	template <typename CharType>
+	struct TJsonValueNumberType
+	{
+		static_assert(sizeof(CharType), "Unsupported type");
+	};
+
+	template <>
+	struct TJsonValueNumberType<TCHAR>
+	{
+		static FString GetType()
+		{
+			return TEXT("NumberString");
+		}
+	};
+
+	template <>
+	struct TJsonValueNumberType<UTF8CHAR>
+	{
+		static FString GetType()
+		{
+			return TEXT("Utf8NumberString");
+		}
+	};
+}
+
+template <typename CharType>
+TJsonValueNumberString<CharType>::TJsonValueNumberString(const TString<CharType>& InString)
+	: Value(InString)
+{
+	Type = EJson::Number;
+}
+
+template <typename CharType>
+TJsonValueNumberString<CharType>::TJsonValueNumberString(TString<CharType>&& InString)
+	: Value(MoveTemp(InString))
+{
+	Type = EJson::Number;
+}
+
+template <typename CharType>
+bool TJsonValueNumberString<CharType>::TryGetString(FString& OutString) const
+{
+	if constexpr (std::is_same_v<FString, TString<CharType>>)
+	{
+		OutString = Value;
+	}
+	else
+	{
+		OutString = FString(Value);
+	}
+	return true;
+}
+
+template <typename CharType>
+bool TJsonValueNumberString<CharType>::TryGetUtf8String(FUtf8String& OutString) const
+{
+	if constexpr (std::is_same_v<FUtf8String, TString<CharType>>)
+	{
+		OutString = Value;
+	}
+	else
+	{
+		OutString = FUtf8String(Value);
+	}
+	return true;
+}
+
+template <typename CharType>
+bool TJsonValueNumberString<CharType>::TryGetNumber(double& OutDouble) const
+{
+	return LexTryParseString(OutDouble, *Value);
+}
+
+template <typename CharType>
+bool TJsonValueNumberString<CharType>::TryGetNumber(float& OutDouble) const
+{
+	return LexTryParseString(OutDouble, *Value);
+}
+
+template <typename CharType>
+bool TJsonValueNumberString<CharType>::TryGetNumber(int8& OutValue) const
+{
+	return LexTryParseString(OutValue, *Value);
+}
+
+template <typename CharType>
+bool TJsonValueNumberString<CharType>::TryGetNumber(int16& OutValue) const
+{
+	return LexTryParseString(OutValue, *Value);
+}
+
+template <typename CharType>
+bool TJsonValueNumberString<CharType>::TryGetNumber(int32& OutValue) const
+{
+	return LexTryParseString(OutValue, *Value);
+}
+
+template <typename CharType>
+bool TJsonValueNumberString<CharType>::TryGetNumber(int64& OutValue) const
+{
+	return LexTryParseString(OutValue, *Value);
+}
+
+template <typename CharType>
+bool TJsonValueNumberString<CharType>::TryGetNumber(uint8& OutValue) const
+{
+	return LexTryParseString(OutValue, *Value);
+}
+
+template <typename CharType>
+bool TJsonValueNumberString<CharType>::TryGetNumber(uint16& OutValue) const
+{
+	return LexTryParseString(OutValue, *Value);
+}
+
+template <typename CharType>
+bool TJsonValueNumberString<CharType>::TryGetNumber(uint32& OutValue) const
+{
+	return LexTryParseString(OutValue, *Value);
+}
+
+template <typename CharType>
+bool TJsonValueNumberString<CharType>::TryGetNumber(uint64& OutValue) const
+{
+	return LexTryParseString(OutValue, *Value);
+}
+
+template <typename CharType>
+bool TJsonValueNumberString<CharType>::TryGetBool(bool& OutBool) const
+{
+	OutBool = Value.ToBool();
+	return true;
+}
+
+template <typename CharType>
+bool TJsonValueNumberString<CharType>::PreferStringRepresentation() const
+{
+	return true;
+}
+
+template <typename CharType>
+SIZE_T TJsonValueNumberString<CharType>::GetMemoryFootprint() const
+{
+	return sizeof(*this) + GetAllocatedSize();
+}
+
+template <typename CharType>
+FString TJsonValueNumberString<CharType>::GetType() const
+{
+	return TJsonValueNumberType<CharType>::GetType();
+}
+
+template <typename CharType>
+SIZE_T TJsonValueNumberString<CharType>::GetAllocatedSize() const
+{
+	return Value.GetAllocatedSize();
+}
diff --git a/Engine/Source/Runtime/Slate/Private/Widgets/Views/SListPanel.cpp b/Engine/Source/Runtime/Slate/Private/Widgets/Views/SListPanel.cpp
--- a/Engine/Source/Runtime/Slate/Private/Widgets/Views/SListPanel.cpp
+++ b/Engine/Source/Runtime/Slate/Private/Widgets/Views/SListPanel.cpp
@@ -25,6 +25,8 @@ SListPanel::SListPanel()
 {
 }
 
+SListPanel::~SListPanel() = default;
+
 void SListPanel::Construct( const FArguments& InArgs )
 {
 	ItemWidthAttribute.Assign(*this, InArgs._ItemWidth);
diff --git a/Engine/Source/Runtime/Slate/Private/Widgets/Views/SListPanel.h b/Engine/Source/Runtime/Slate/Private/Widgets/Views/SListPanel.h
--- a/Engine/Source/Runtime/Slate/Private/Widgets/Views/SListPanel.h
+++ b/Engine/Source/Runtime/Slate/Private/Widgets/Views/SListPanel.h
@@ -61,6 +61,8 @@ public:
 	
 	SLATE_API SListPanel();
 
+	SLATE_API virtual ~SListPanel();
+
 	/**
 	 * Construct the widget
 	 *
diff --git a/Engine/Source/Runtime/SlateCore/Private/Widgets/SLeafWidget.cpp b/Engine/Source/Runtime/SlateCore/Private/Widgets/SLeafWidget.cpp
--- a/Engine/Source/Runtime/SlateCore/Private/Widgets/SLeafWidget.cpp
+++ b/Engine/Source/Runtime/SlateCore/Private/Widgets/SLeafWidget.cpp
@@ -6,6 +6,13 @@
 /* SLeafWidget interface
  *****************************************************************************/
 
+SLeafWidget::SLeafWidget()
+{
+	bCanHaveChildren = false;
+}
+
+SLeafWidget::~SLeafWidget() = default;
+
 void SLeafWidget::SetVisibility( TAttribute<EVisibility> InVisibility )
 {
 	SWidget::SetVisibility( InVisibility );
diff --git a/Engine/Source/Runtime/SlateCore/Public/Widgets/SLeafWidget.h b/Engine/Source/Runtime/SlateCore/Public/Widgets/SLeafWidget.h
--- a/Engine/Source/Runtime/SlateCore/Public/Widgets/SLeafWidget.h
+++ b/Engine/Source/Runtime/SlateCore/Public/Widgets/SLeafWidget.h
@@ -22,10 +22,9 @@ class SLeafWidget
 	: public SWidget
 {
 public:
-	SLeafWidget()
-	{
-		bCanHaveChildren = false;
-	}
+	SLATECORE_API SLeafWidget();
+
+	SLATECORE_API virtual ~SLeafWidget();
 
 	SLATECORE_API virtual void SetVisibility( TAttribute<EVisibility> InVisibility ) override final;
 
@@ -43,7 +42,7 @@ private:
 	 *
 	 * Whenever possible, LeafWidgets should avoid dealing with layout properties. See TextBlock for an example.
 	 */
-	virtual int32 OnPaint( const FPaintArgs& Args, const FGeometry& AllottedGeometry, const FSlateRect& MyCullingRect, FSlateWindowElementList& OutDrawElements, int32 LayerId, const FWidgetStyle& InWidgetStyle, bool bParentEnabled ) const override = 0;
+	SLATECORE_API virtual int32 OnPaint( const FPaintArgs& Args, const FGeometry& AllottedGeometry, const FSlateRect& MyCullingRect, FSlateWindowElementList& OutDrawElements, int32 LayerId, const FWidgetStyle& InWidgetStyle, bool bParentEnabled ) const override = 0;
 	
 	/**
 	 * Overwritten from SWidget.
@@ -52,7 +51,7 @@ private:
 	 * take child widgets into account as LeafWidgets have none by definition. For example, the TextBlock widget simply
 	 * measures the area necessary to display its text with the given font and font size.
 	 */
-	virtual FVector2D ComputeDesiredSize(float) const override = 0;
+	SLATECORE_API virtual FVector2D ComputeDesiredSize(float) const override = 0;
 	
 	/**
 	 * Overwritten from SWidget.
