Add support for writing TMaps as JSON objects.
#rnx

diff --git a/Engine/Source/Runtime/Json/Public/Serialization/JsonWriter.h b/Engine/Source/Runtime/Json/Public/Serialization/JsonWriter.h
--- a/Engine/Source/Runtime/Json/Public/Serialization/JsonWriter.h
+++ b/Engine/Source/Runtime/Json/Public/Serialization/JsonWriter.h
@@ -260,6 +260,28 @@ public:
 		WriteArrayEnd();
 	}
 
+	template<typename IdentifierType, typename MapIdentifierType, class MapElementType>
+	void WriteValue(IdentifierType&& Identifier, const TMap<MapIdentifierType, MapElementType>& Map)
+	{
+		WriteObjectStart(Forward<IdentifierType>(Identifier));
+		for (const TPair<MapIdentifierType, MapElementType>& Element : Map)
+		{
+			WriteValue(Element.Key, Element.Value);
+		}
+		WriteObjectEnd();
+	}
+
+	template<typename MapIdentifierType, class MapElementType>
+	void WriteValue(const TMap<MapIdentifierType, MapElementType>& Map)
+	{
+		WriteObjectStart();
+		for (const TPair<MapIdentifierType, MapElementType>& Element : Map)
+		{
+			WriteValue(Element.Key, Element.Value);
+		}
+		WriteObjectEnd();
+	}
+
 	void WriteValue(FStringView Identifier, const TCHAR* Value)
 	{
 		WriteValue(Identifier, FStringView(Value));
