Rename JSON_SERIALIZE_SERIALIZABLE to JSON_SERIALIZE_MEMBERS_OF.

diff --git a/Engine/Plugins/Online/OnlineSubsystem/Source/Public/OnlineJsonSerializer.h b/Engine/Plugins/Online/OnlineSubsystem/Source/Public/OnlineJsonSerializer.h
--- a/Engine/Plugins/Online/OnlineSubsystem/Source/Public/OnlineJsonSerializer.h
+++ b/Engine/Plugins/Online/OnlineSubsystem/Source/Public/OnlineJsonSerializer.h
@@ -13,7 +13,7 @@ class FVariantData;
 #define ONLINE_JSON_SERIALIZE_ARRAY(JsonName, JsonArray)											JSON_SERIALIZE_ARRAY(JsonName, JsonArray) 
 #define ONLINE_JSON_SERIALIZE_MAP(JsonName, JsonMap)												JSON_SERIALIZE_MAP(JsonName, JsonMap)
 #define ONLINE_JSON_SERIALIZE_SIMPLE_COPY(JsonMap)													JSON_SERIALIZE_SIMPLECOPY(JsonMap)
-#define ONLINE_JSON_SERIALIZE_SERIALIZABLE(JsonName, JsonValue)										JSON_SERIALIZE_SERIALIZABLE(JsonName, JsonValue) 
+#define ONLINE_JSON_SERIALIZE_SERIALIZABLE(JsonName, JsonValue)	UE_DEPRECATED_MACRO(5.6, "ONLINE_JSON_SERIALIZE_SERIALIZABLE has been deprecated, use JSON_SERIALIZE_MEMBERS_OF instead") JSON_SERIALIZE_MEMBERS_OF(JsonValue) 
 #define ONLINE_JSON_SERIALIZE_RAW_JSON_STRING(JsonName, JsonValue)									JSON_SERIALIZE_RAW_JSON_STRING(JsonName, JsonValue)
 #define ONLINE_JSON_SERIALIZE_ARRAY_SERIALIZABLE(JsonName, JsonArray, ElementType)					JSON_SERIALIZE_ARRAY_SERIALIZABLE(JsonName, JsonArray, ElementType)
 #define ONLINE_JSON_SERIALIZE_OPTIONAL_ARRAY_SERIALIZABLE(JsonName, OptionalJsonArray, ElementType)	JSON_SERIALIZE_OPTIONAL_ARRAY_SERIALIZABLE(JsonName, OptionalJsonArray, ElementType)
diff --git a/Engine/Source/Runtime/Json/Private/Tests/JsonTests.cpp b/Engine/Source/Runtime/Json/Private/Tests/JsonTests.cpp
--- a/Engine/Source/Runtime/Json/Private/Tests/JsonTests.cpp
+++ b/Engine/Source/Runtime/Json/Private/Tests/JsonTests.cpp
@@ -8,7 +8,9 @@
 #include "Policies/CondensedJsonPrintPolicy.h"
 #include "Policies/PrettyJsonPrintPolicy.h"
 #include "Serialization/JsonReader.h"
+#include "Serialization/JsonSerializable.h"
 #include "Serialization/JsonSerializer.h"
+#include "Serialization/JsonSerializerMacros.h"
 #include "Serialization/JsonTypes.h"
 #include "Tests/TestHarnessAdapter.h"
 
@@ -1466,4 +1468,47 @@ TEST_CASE_NAMED(FJsonAutomationTest, "System::Engine::FileSystem::JSON", "[Appli
 	}
 }
 
+TEST_CASE_NAMED(FJsonSerializerMacrosTest, "System::Engine::FileSystem::JSON::MacroToSerializeMembers", "[ApplicationContextMask][SmokeFilter]")
+{
+	struct FJsonStruct : public FJsonSerializable
+	{
+		struct FSubJsonStruct : public FJsonSerializable
+		{
+			int32 SubVarInt;
+			FString SubVarString;
+
+			FSubJsonStruct(int32 InSubVarInt = 0, FStringView InSubVarString = TEXTVIEW(""))
+				: SubVarInt(InSubVarInt)
+				, SubVarString(InSubVarString)
+			{
+			}
+
+			BEGIN_JSON_SERIALIZER
+				JSON_SERIALIZE("sub_var_int", SubVarInt);
+				JSON_SERIALIZE("sub_var_string", SubVarString);
+			END_JSON_SERIALIZER
+		};
+
+		int32 VarInt = 0;
+		FSubJsonStruct VarSerializable;
+
+		BEGIN_JSON_SERIALIZER
+			JSON_SERIALIZE("var_int", VarInt);
+			JSON_SERIALIZE_MEMBERS_OF(VarSerializable);
+		END_JSON_SERIALIZER;
+	};
+
+	const TCHAR* SourceJsonString = TEXT("{\"var_int\":2,\"sub_var_int\":10,\"sub_var_string\":\"abc\"}");
+
+	FJsonStruct JsonStructTarget;
+	JsonStructTarget.FromJson(SourceJsonString);
+
+	CHECK(JsonStructTarget.VarInt == 2);
+	CHECK(JsonStructTarget.VarSerializable.SubVarInt == 10);
+	CHECK(JsonStructTarget.VarSerializable.SubVarString == TEXT("abc"));
+
+	FString ToJsonResult = JsonStructTarget.ToJson(/*bPrettyPrint*/false);
+	CHECK(ToJsonResult == SourceJsonString);
+}
+
 #endif // WITH_TESTS
diff --git a/Engine/Source/Runtime/Json/Public/Serialization/JsonSerializerMacros.h b/Engine/Source/Runtime/Json/Public/Serialization/JsonSerializerMacros.h
--- a/Engine/Source/Runtime/Json/Public/Serialization/JsonSerializerMacros.h
+++ b/Engine/Source/Runtime/Json/Public/Serialization/JsonSerializerMacros.h
@@ -110,9 +110,6 @@
 #define JSON_SERIALIZE_MAP_SAFE(JsonName, JsonMap) \
 		Serializer.SerializeMapSafe(TEXTVIEW(JsonName), JsonMap)
 
-#define JSON_SERIALIZE_SERIALIZABLE(JsonName, JsonValue) \
-		JsonValue.Serialize(Serializer, false)
-
 #define JSON_SERIALIZE_RAW_JSON_STRING(JsonName, JsonValue) \
 		if (Serializer.IsLoading()) \
 		{ \
@@ -272,6 +269,13 @@
 			Serializer.EndObject(); \
 		}
 
+// Use this to read/write JsonValue with its own keys and values corresponding to its members
+#define JSON_SERIALIZE_MEMBERS_OF(JsonValue) \
+		JsonValue.Serialize(Serializer, /*bIsFlatObject*/true)
+
+#define JSON_SERIALIZE_SERIALIZABLE(JsonName, JsonValue) UE_DEPRECATED_MACRO(5.6, "JSON_SERIALIZE_SERIALIZABLE has been deprecated, use JSON_SERIALIZE_MEMBERS_OF instead") JSON_SERIALIZE_MEMBERS_OF(JsonValue)
+
+// Use this when JsonSerializableObject will be read/write as json object, while JsonName is the key
 #define JSON_SERIALIZE_OBJECT_SERIALIZABLE(JsonName, JsonSerializableObject) \
 		/* Process the JsonName field differently because it is an object */ \
 		if (Serializer.IsLoading()) \
