Changes to the Json Writer to support TNestedVariant

diff --git a/Engine/Source/Runtime/Json/Public/Serialization/JsonWriter.h b/Engine/Source/Runtime/Json/Public/Serialization/JsonWriter.h
--- a/Engine/Source/Runtime/Json/Public/Serialization/JsonWriter.h
+++ b/Engine/Source/Runtime/Json/Public/Serialization/JsonWriter.h
@@ -111,7 +111,7 @@ public:
 		check(CanWriteObjectWithoutIdentifier());
 		WriteCommaIfNeeded();
 
-		if ( PreviousTokenWritten != EJsonToken::None )
+		if ( PreviousTokenWritten != EJsonToken::Comma && PreviousTokenWritten != EJsonToken::None )
 		{
 			PrintPolicy::WriteLineTerminator(Stream);
 			PrintPolicy::WriteTabs(Stream, IndentLevel);
@@ -155,7 +155,7 @@ public:
 		check(CanWriteValueWithoutIdentifier());
 		WriteCommaIfNeeded();
 
-		if ( PreviousTokenWritten != EJsonToken::None )
+		if ( PreviousTokenWritten != EJsonToken::Comma && PreviousTokenWritten != EJsonToken::None )
 		{
 			PrintPolicy::WriteLineTerminator(Stream);
 			PrintPolicy::WriteTabs(Stream, IndentLevel);
@@ -200,15 +200,34 @@ public:
 		PreviousTokenWritten = EJsonToken::SquareClose;
 	}
 
+	// Special case for bit fields
+	void WriteValue(uint8 Value)
+	{
+		check(CanWriteValueWithoutIdentifier());
+		WriteCommaIfNeeded();
+
+		if (EJsonToken_IsShortValue(PreviousTokenWritten))
+		{
+			PrintPolicy::WriteSpace(Stream);
+		}
+		else
+		{
+			PrintPolicy::WriteLineTerminator(Stream);
+			PrintPolicy::WriteTabs(Stream, IndentLevel);
+		}
+
+		PreviousTokenWritten = WriteValueOnly(Value);
+	}
+
 	template <class FValue>
-	void WriteValue(FValue Value)
+	void WriteValue(FValue&& Value)
 	{
 		check(CanWriteValueWithoutIdentifier());
 		WriteCommaIfNeeded();
 
-		if (PreviousTokenWritten == EJsonToken::SquareOpen || EJsonToken_IsShortValue(PreviousTokenWritten))
+		if (EJsonToken_IsShortValue(PreviousTokenWritten))
 		{
-			PrintPolicy::WriteSpace( Stream );
+			PrintPolicy::WriteSpace(Stream);
 		}
 		else
 		{
@@ -216,7 +235,7 @@ public:
 			PrintPolicy::WriteTabs(Stream, IndentLevel);
 		}
 
-		PreviousTokenWritten = WriteValueOnly( Value );
+		PreviousTokenWritten = WriteValueOnly(Forward<FValue>(Value));
 	}
 
 	void WriteValue(FStringView Value)
@@ -239,14 +258,25 @@ public:
 		PreviousTokenWritten = WriteValueOnly(Value);
 	}
 
+	// Special case for bit fields
+	template<typename IdentifierType>
+	void WriteValue(IdentifierType&& Identifier, uint8 Value)
+	{
+		check( Stack.Top() == EJson::Object );
+		WriteIdentifier(Forward<IdentifierType>(Identifier));
+
+		PrintPolicy::WriteSpace(Stream);
+		PreviousTokenWritten = WriteValueOnly(Value);
+	}
+
 	template<class FValue, typename IdentifierType>
-	void WriteValue(IdentifierType&& Identifier, FValue Value)
+	void WriteValue(IdentifierType&& Identifier, FValue&& Value)
 	{
 		check( Stack.Top() == EJson::Object );
 		WriteIdentifier(Forward<IdentifierType>(Identifier));
 
 		PrintPolicy::WriteSpace(Stream);
-		PreviousTokenWritten = WriteValueOnly(MoveTemp(Value));
+		PreviousTokenWritten = WriteValueOnly(Forward<FValue>(Value));
 	}
 
 	template<class ElementType, typename IdentifierType>
@@ -380,9 +410,11 @@ protected:
 
 	inline void WriteCommaIfNeeded()
 	{
-		if ( PreviousTokenWritten != EJsonToken::CurlyOpen && PreviousTokenWritten != EJsonToken::SquareOpen && PreviousTokenWritten != EJsonToken::Identifier && PreviousTokenWritten != EJsonToken::None)
+		if ( PreviousTokenWritten != EJsonToken::CurlyOpen && PreviousTokenWritten != EJsonToken::SquareOpen && PreviousTokenWritten != EJsonToken::Identifier && PreviousTokenWritten != EJsonToken::Comma && PreviousTokenWritten != EJsonToken::None)
 		{
 			PrintPolicy::WriteChar(Stream, CharType(','));
+
+			PreviousTokenWritten = EJsonToken::Comma;
 		}
 	}
 
@@ -406,6 +438,8 @@ protected:
 		PrintPolicy::WriteTabs(Stream, IndentLevel);
 		WriteStringValue(Identifier);
 		PrintPolicy::WriteChar(Stream, CharType(':'));
+
+		PreviousTokenWritten = EJsonToken::Identifier;
 	}
 
 	void WriteIdentifier(const FText& Identifier)
@@ -421,6 +455,8 @@ protected:
 		PrintPolicy::WriteTabs(Stream, IndentLevel);
 		WriteStringValue(FStringView(Identifier));
 		PrintPolicy::WriteChar(Stream, CharType(':'));
+
+		PreviousTokenWritten = EJsonToken::Identifier;
 	}
 
 	inline EJsonToken WriteValueOnly(bool Value)
@@ -454,6 +490,12 @@ protected:
 		return EJsonToken::Number;
 	}
 
+	// Special case for bit fields
+	inline EJsonToken WriteValueOnly(uint8 Value)
+	{
+		return WriteValueOnly((uint64)Value);
+	}
+
 	inline EJsonToken WriteValueOnly(uint32 Value)
 	{
 		return WriteValueOnly((uint64)Value);
@@ -489,6 +531,60 @@ protected:
 		return EJsonToken::String;
 	}
 
+	template<typename ValueType>
+	EJsonToken WriteValueOnly(TSharedRef<ValueType> ValueRef)
+	{
+		return WriteValueOnly(*ValueRef);
+	}
+
+	template<typename ValueType>
+	EJsonToken WriteValueOnly(TSharedPtr<ValueType> ValuePtr)
+	{
+		if(ValuePtr.IsValid())
+		{
+			return WriteValueOnly(*ValuePtr);
+		}
+		else
+		{
+			return WriteValueOnly(nullptr);
+		}
+	}
+
+	template<typename... ValueTypes>
+	EJsonToken WriteValueOnly(const TVariant<ValueTypes...>& Variant)
+	{
+		return Visit([this](const auto& Value)
+			{
+				return WriteValueOnly(Value);
+			}, Variant);
+	}
+
+	template<typename MapIdentifierType, typename MapElementType>
+	EJsonToken WriteValueOnly(const TMap<MapIdentifierType, MapElementType>& Map)
+	{
+		WriteObjectStart();
+		for (const TPair<MapIdentifierType, MapElementType>& Element : Map)
+		{
+			WriteValue(Element.Key, Element.Value);
+		}
+		WriteObjectEnd();
+
+		return PreviousTokenWritten;
+	}
+
+	template<typename ArrayElementType>
+	EJsonToken WriteValueOnly(const TArray<ArrayElementType>& Array)
+	{
+		WriteArrayStart();
+		for (int Idx = 0; Idx < Array.Num(); Idx++)
+		{
+			WriteValue(Array[Idx]);
+		}
+		WriteArrayEnd();
+
+		return PreviousTokenWritten;
+	}
+
 	virtual void WriteStringValue(FAnsiStringView String)
 	{
 		PrintPolicy::WriteChar(Stream, CharType('"'));
