Missing token update in JsonWriter

diff --git a/Engine/Source/Runtime/Json/Public/Serialization/JsonWriter.h b/Engine/Source/Runtime/Json/Public/Serialization/JsonWriter.h
--- a/Engine/Source/Runtime/Json/Public/Serialization/JsonWriter.h
+++ b/Engine/Source/Runtime/Json/Public/Serialization/JsonWriter.h
@@ -427,6 +427,8 @@ protected:
 		PrintPolicy::WriteTabs(Stream, IndentLevel);
 		WriteStringValue(TStringView<InCharType>(Identifier));
 		PrintPolicy::WriteChar(Stream, CharType(':'));
+
+		PreviousTokenWritten = EJsonToken::Identifier;
 	}
 	
 	template <typename InCharType>
