Small modification to JsonTests expected results after JsonWriter changes.
Related tweaks to JsonWriter.

diff --git a/Engine/Source/Runtime/Json/Private/Tests/JsonTests.cpp b/Engine/Source/Runtime/Json/Private/Tests/JsonTests.cpp
--- a/Engine/Source/Runtime/Json/Private/Tests/JsonTests.cpp
+++ b/Engine/Source/Runtime/Json/Private/Tests/JsonTests.cpp
@@ -1110,7 +1110,12 @@ TEST_CASE_NAMED(FJsonAutomationTest, "System::Engine::FileSystem::JSON", "[Appli
 					"			\"InnerData1\": \"value\""	LINE_TERMINATOR_ANSI
 					"		},"								LINE_TERMINATOR_ANSI
 					"		[],"							LINE_TERMINATOR_ANSI
-					"		[ 1, 2, 3, 4 ],"				LINE_TERMINATOR_ANSI
+					"		["								LINE_TERMINATOR_ANSI
+					"			1,"							LINE_TERMINATOR_ANSI
+					"			2,"							LINE_TERMINATOR_ANSI
+					"			3,"							LINE_TERMINATOR_ANSI
+					"			4"							LINE_TERMINATOR_ANSI		
+					"		],"								LINE_TERMINATOR_ANSI
 					"		{"								LINE_TERMINATOR_ANSI
 					"		},"								LINE_TERMINATOR_ANSI
 					"		\"value\","						LINE_TERMINATOR_ANSI
diff --git a/Engine/Source/Runtime/Json/Public/Serialization/JsonWriter.h b/Engine/Source/Runtime/Json/Public/Serialization/JsonWriter.h
--- a/Engine/Source/Runtime/Json/Public/Serialization/JsonWriter.h
+++ b/Engine/Source/Runtime/Json/Public/Serialization/JsonWriter.h
@@ -111,7 +111,7 @@ public:
 		check(CanWriteObjectWithoutIdentifier());
 		WriteCommaIfNeeded();
 
-		if ( PreviousTokenWritten != EJsonToken::Comma && PreviousTokenWritten != EJsonToken::None )
+		if (PreviousTokenWritten != EJsonToken::None )
 		{
 			PrintPolicy::WriteLineTerminator(Stream);
 			PrintPolicy::WriteTabs(Stream, IndentLevel);
@@ -155,7 +155,7 @@ public:
 		check(CanWriteValueWithoutIdentifier());
 		WriteCommaIfNeeded();
 
-		if ( PreviousTokenWritten != EJsonToken::Comma && PreviousTokenWritten != EJsonToken::None )
+		if (PreviousTokenWritten != EJsonToken::None )
 		{
 			PrintPolicy::WriteLineTerminator(Stream);
 			PrintPolicy::WriteTabs(Stream, IndentLevel);
@@ -185,16 +185,13 @@ public:
 		check( Stack.Top() == EJson::Array );
 
 		--IndentLevel;
-		if ( PreviousTokenWritten == EJsonToken::SquareClose || PreviousTokenWritten == EJsonToken::CurlyClose || PreviousTokenWritten == EJsonToken::String )
+
+		if (PreviousTokenWritten != EJsonToken::SquareOpen)
 		{
 			PrintPolicy::WriteLineTerminator(Stream);
 			PrintPolicy::WriteTabs(Stream, IndentLevel);
 		}
-		else if ( PreviousTokenWritten != EJsonToken::SquareOpen )
-		{
-			PrintPolicy::WriteSpace( Stream );
-		}
-
+		
 		PrintPolicy::WriteChar(Stream, CharType(']'));
 		Stack.Pop();
 		PreviousTokenWritten = EJsonToken::SquareClose;
