Added write helpers to RapidJsonUtils

* Added UE::Json::WriteCompact and UE::Json::WritePretty, to write a UE::Json::FDocument to a FString in either compact or pretty form.
* Added UE::Json::MakeStringRef and UE::Json::MakeStringValue to help convert a FStringView to a UE::Json::FStringRef or UE::Json::FValue when building a JSON DOM.
* Added some additional aliases to UE::Json that are needed when building a JSON DOM.
* Switched UE::Json aliases to use the UE FMemory allocator, rather than rapidjson::CrtAllocator.

diff --git a/Engine/Source/Runtime/Json/Private/JsonUtils/RapidJsonUtils.cpp b/Engine/Source/Runtime/Json/Private/JsonUtils/RapidJsonUtils.cpp
--- a/Engine/Source/Runtime/Json/Private/JsonUtils/RapidJsonUtils.cpp
+++ b/Engine/Source/Runtime/Json/Private/JsonUtils/RapidJsonUtils.cpp
@@ -1,10 +1,9 @@
 // Copyright Epic Games, Inc. All Rights Reserved.
 
 #include "JsonUtils/RapidJsonUtils.h"
-#include "Dom/JsonValue.h"
 
-namespace UE {
-namespace Json {
+namespace UE::Json
+{
 
 const constexpr uint32 DefaultParseFlags = rapidjson::ParseFlag::kParseTrailingCommasFlag;
 
@@ -272,6 +271,27 @@ TValueOrError<FDocument, FParseError> ParseInPlace(TArrayView<TCHAR> JsonText)
 	return MakeValue(MoveTemp(Result));
 }
 
+FString WriteCompact(const FDocument& Document)
+{
+	FStringBuffer StringBuffer;
+	FStringWriter StringWriter(StringBuffer);
+
+	Document.Accept(StringWriter);
+
+	return StringBuffer.GetString();
+}
+
+FString WritePretty(const FDocument& Document)
+{
+	FStringBuffer StringBuffer;
+	FPrettyStringWriter StringWriter(StringBuffer);
+	StringWriter.SetIndent('\t', 1);
+
+	Document.Accept(StringWriter);
+
+	return StringBuffer.GetString();
+}
+
 const TCHAR* GetValueTypeName(const FValue& Value)
 {
 	switch (Value.GetType())
@@ -300,5 +320,4 @@ const TCHAR* GetValueTypeName(const FValue& Value)
 	}
 }
 
-} // namespace Json
-} // namespace UE
+} // namespace UE::Json
diff --git a/Engine/Source/Runtime/Json/Public/JsonUtils/RapidJsonUtils.h b/Engine/Source/Runtime/Json/Public/JsonUtils/RapidJsonUtils.h
--- a/Engine/Source/Runtime/Json/Public/JsonUtils/RapidJsonUtils.h
+++ b/Engine/Source/Runtime/Json/Public/JsonUtils/RapidJsonUtils.h
@@ -5,20 +5,89 @@
 #include "RapidJsonIncludes.h"
 #include "Misc/Optional.h"
 #include "Containers/StringView.h"
+#include "Containers/UnrealString.h"
 #include "Templates/ValueOrError.h"
 
 /**
- * Helper functions for reading from RapidJSON into unreal types
+ * Helper functions for reading/writing RapidJSON with unreal types
  */
-namespace UE {
-namespace Json {
+namespace UE::Json
+{
+namespace Private
+{
+	// TODO: Can remove this and just use rapidjson::CrtAllocator once RapidJsonIncludes.h starts overriding RAPIDJSON_MALLOC, etc
+	class FAllocatorImpl
+	{
+	public:
+		static const bool kNeedFree = true;
+
+		void* Malloc(size_t Count)
+		{
+			return FMemory::Malloc(Count);
+		}
+
+		void* Realloc(void* Original, size_t, size_t Count)
+		{
+			if (Count == 0)
+			{
+				FMemory::Free(Original);
+				return nullptr;
+			}
+			return FMemory::Realloc(Original, Count);
+		}
+
+		static void Free(void* Ptr)
+		{
+			FMemory::Free(Ptr);
+		}
 
-	// types that use TCHAR
-	using FDocument = rapidjson::GenericDocument<rapidjson::UTF16<TCHAR>, rapidjson::CrtAllocator>;
+		bool operator==(const FAllocatorImpl&) const
+		{
+			return true;
+		}
+
+		bool operator!=(const FAllocatorImpl&) const
+		{
+			return false;
+		}
+	};
+} // namespace UE::Json::Private
+
+	// types that use TCHAR (common)
+	using FEncoding = rapidjson::UTF16<TCHAR>;
+	using FAllocator = Private::FAllocatorImpl;
+	using FDocument = rapidjson::GenericDocument<FEncoding, FAllocator>;
+	using FStringRef = FDocument::StringRefType;
 	using FValue = FDocument::ValueType;
+	using FMember = FValue::Member;
+	
+	// types that use TCHAR (for reading)
 	using FConstObject = FValue::ConstObject;
 	using FConstArray = FValue::ConstArray;
-	using FMember = FValue::Member;
+
+	// types that use TCHAR (for writing)
+	using FObject = FValue::Object;
+	using FArray = FValue::Array;
+	using FStringBuffer = rapidjson::GenericStringBuffer<FEncoding, FAllocator>;
+	using FStringWriter = rapidjson::Writer<FStringBuffer, FEncoding, FEncoding, FAllocator>;
+	using FPrettyStringWriter = rapidjson::PrettyWriter<FStringBuffer, FEncoding, FEncoding, FAllocator>;
+
+	/**
+	 * Make a RapidJSON reference to the given string.
+	 * @note The given string must outlive the returned reference.
+	 */
+	inline FStringRef MakeStringRef(FStringView Str)
+	{
+		return FStringRef(Str.GetData(), Str.Len());
+	}
+
+	/**
+	 * Make a RapidJSON value that has a copy of the given string.
+	 */
+	inline FValue MakeStringValue(FStringView Str, FAllocator& Allocator)
+	{
+		return FValue(Str.GetData(), Str.Len(), Allocator);
+	}
 
 	/**
 	 * Attempts to read a boolean with a specified name.
@@ -89,7 +158,7 @@ namespace Json {
 	}
 
 	/**
-	 * Attempts to find an object with the specified namennn
+	 * Attempts to find an object with the specified name
 	 * 
 	 * Returns a valid object ref if the field was found and it matches an object.
 	 */
@@ -118,7 +187,7 @@ namespace Json {
 		size_t Offset = 0;
 
 		/**
-		 * Create a user-readable error message. Takes the source text paramter to find the line number from the offset
+		 * Create a user-readable error message. Takes the source text parameter to find the line number from the offset
 		 */
 		JSON_API FString CreateMessage(const FStringView JsonText) const;
 	};
@@ -135,9 +204,19 @@ namespace Json {
 	 */
 	JSON_API TValueOrError<FDocument, FParseError> ParseInPlace(TArrayView<TCHAR> JsonText);
 
+	/**
+	 * Write a RapidJSON document as a compact string with default flags.
+	 */
+	JSON_API FString WriteCompact(const FDocument& Document);
+
+	/**
+	 * Write a RapidJSON document as a pretty string with default flags and tab indentation.
+	 */
+	JSON_API FString WritePretty(const FDocument& Document);
+
 	/**
 	 * Return a string representation of type name, for debug logging purposes
 	 */
 	JSON_API const TCHAR* GetValueTypeName(const FValue& Value);
-}
-}
+
+} // namespace UE::Json
diff --git a/Engine/Source/Runtime/Json/Public/RapidJsonIncludes.h b/Engine/Source/Runtime/Json/Public/RapidJsonIncludes.h
--- a/Engine/Source/Runtime/Json/Public/RapidJsonIncludes.h
+++ b/Engine/Source/Runtime/Json/Public/RapidJsonIncludes.h
@@ -59,6 +59,7 @@ THIRD_PARTY_INCLUDES_START
 #include "rapidjson/document.h"
 #include "rapidjson/encodings.h"
 #include "rapidjson/writer.h"
+#include "rapidjson/prettywriter.h"
 #include "rapidjson/stringbuffer.h"
 #include "rapidjson/schema.h"
 // just for non-user facing logging purposes use the English error descriptions
