Refactor of JsonSerializer to be able to set custom policies and read data to types other than JsonObject/Value

#FailedValidations Preflight

diff --git a/Engine/Source/Runtime/CoreOnline/Private/Online/NestedVariantJson.cpp b/Engine/Source/Runtime/CoreOnline/Private/Online/NestedVariantJson.cpp
--- a/Engine/Source/Runtime/CoreOnline/Private/Online/NestedVariantJson.cpp
+++ b/Engine/Source/Runtime/CoreOnline/Private/Online/NestedVariantJson.cpp
@@ -6,189 +6,160 @@
 #include "Policies/CondensedJsonPrintPolicy.h"
 #include "Serialization/JsonSerializer.h"
 
-// TODO: Update the JSON module to handle nested variants of any type directly so this file is not needed.
+using FJsonWriterFactory = TJsonWriterFactory<TCHAR, TCondensedJsonPrintPolicy<TCHAR>>;
+using FJsonWriter = TJsonWriter<TCHAR, TCondensedJsonPrintPolicy<TCHAR>>;
+using FJsonSerializerNestedVariant = TJsonSerializer<FJsonSerializerPolicy_NestedVariant>;
 
 FString NestedVariantToJson(const FNestedVariantJson::FMapPtr& Map)
 {
 	FString Result;
+
 	if (!Map)
 	{
 		return Result;
 	}
 
-	using FJsonWriterFactory = TJsonWriterFactory<TCHAR, TCondensedJsonPrintPolicy<TCHAR>>;
-	using FJsonWriter = TJsonWriter<TCHAR, TCondensedJsonPrintPolicy<TCHAR>>;
-
-	const FNestedVariantJson::FMapRef MapRef = Map.ToSharedRef();
-	TSharedRef<FJsonObject> JsonObject = NestedVariantToJsonObject(MapRef);
 	TSharedRef<FJsonWriter> Writer = FJsonWriterFactory::Create(&Result);
-	verify(FJsonSerializer::Serialize(JsonObject, Writer));
+	verify(FJsonSerializerNestedVariant::Serialize(Map, Writer));
+
 	return Result;
 }
 
 TSharedRef<FJsonObject> NestedVariantToJsonObject(const FNestedVariantJson::FMapRef& Map)
 {
-	TSharedRef<FJsonObject> JsonObject = MakeShared<FJsonObject>();
-	JsonObject->Values.Reserve(Map->Num());
+	FString JsonStr = NestedVariantToJson(Map);
+	TSharedRef<TJsonReader<TCHAR>> Reader = TJsonReaderFactory<TCHAR>::Create(JsonStr);	
+	TSharedPtr<FJsonObject> JsonObject = MakeShared<FJsonObject>();
+	verify(FJsonSerializer::Deserialize(Reader, JsonObject));
 
-	for (const TPair<FString, FNestedVariantJson::FValue>& Pair : *Map)
+	return JsonObject.ToSharedRef();
+}
+
+void NestedVariantFromJson(const char* InJson, FNestedVariantJson::FMapRef& Map)
+{
+	if (!InJson || strlen(InJson) == 0)
 	{
-		JsonObject->SetField(Pair.Key, NestedVariantToJsonValue(Pair.Value));
+		return;
 	}
 
-	return JsonObject;
+	TSharedRef<TJsonReader<TCHAR>> Reader = TJsonReaderFactory<TCHAR>::Create(UTF8_TO_TCHAR(InJson));
+	FNestedVariantJson::FMapPtr MapPtr =  Map.ToSharedPtr();
+	verify(FJsonSerializerNestedVariant::Deserialize(Reader, MapPtr));
+}
+
+void NestedVariantFromJsonObject(const TSharedPtr<FJsonObject>& JsonObject, FNestedVariantJson::FMapRef& Map)
+{
+	FString JsonStr;
+	TSharedRef<FJsonWriter> Writer = FJsonWriterFactory::Create(&JsonStr);
+	verify(FJsonSerializer::Serialize(JsonObject, Writer));
+
+	NestedVariantFromJson(TCHAR_TO_UTF8(*JsonStr), Map);
 }
 
-TSharedPtr<FJsonValue> NestedVariantToJsonValue(const FNestedVariantJson::FValue& Value)
+void LexFromString(FNestedVariantJson::FMap& OutValue, const TCHAR* InString)
 {
-	switch (Value.GetIndex())
+	FNestedVariantJson::FMapRef  OutValueRef = FNestedVariantJson::FMap::CreateVariant();
+	NestedVariantFromJson(TCHAR_TO_UTF8(InString), OutValueRef);
+
+	OutValue = OutValueRef.Get();
+}
+
+bool FJsonSerializerPolicy_NestedVariant::GetValueFromState(const StackState& State, FValue& OutValue)
+{
+	switch (State.Type)
 	{
-	case FNestedVariantJson::FValue::IndexOfType<bool>():
-		return MakeShared<FJsonValueBoolean>(Value.Get<bool>());
-	case FNestedVariantJson::FValue::IndexOfType<int64>():
-		return MakeShared<FJsonValueNumber>(Value.Get<int64>());
-	case FNestedVariantJson::FValue::IndexOfType<double>():
-		return MakeShared<FJsonValueNumber>(Value.Get<double>());
-	case FNestedVariantJson::FValue::IndexOfType<FString>():
-		return MakeShared<TJsonValueString<TCHAR>>(Value.Get<FString>());
-	case FNestedVariantJson::FValue::IndexOfType<FNestedVariantJson::FArrayRef>():
+	case EJson::Object:
 	{
-		const FNestedVariantJson::FArrayRef& Array = Value.Get<FNestedVariantJson::FArrayRef>();
-		TArray<TSharedPtr<FJsonValue>> JsonArray;
-		JsonArray.Reserve(Array->Num());
-
-		for (const FNestedVariantJson::FValue& ArrayValue : *Array)
+		if (State.Object.IsValid())
 		{
-			JsonArray.Add(NestedVariantToJsonValue(ArrayValue));
+			OutValue.Set<FNestedVariantJson::FMapRef>(State.Object.ToSharedRef());
+			return true;
+		}
+	}		
+	case EJson::Array:
+	{
+		if (State.Array.IsValid())
+		{
+			OutValue.Set<FNestedVariantJson::FArrayRef>(State.Array.ToSharedRef());
+			return true;
 		}
-
-		return MakeShared<FJsonValueArray>(MoveTemp(JsonArray));
 	}
-	case FNestedVariantJson::FValue::IndexOfType<FNestedVariantJson::FMapRef>():
-		return MakeShared<FJsonValueObject>(NestedVariantToJsonObject(Value.Get<FNestedVariantJson::FMapRef>()));
+	case EJson::None:
+	case EJson::Null:
+	case EJson::String:
+	case EJson::Number:
+	case EJson::Boolean:
 	default:
-		checkNoEntry();
+		// FIXME: would be nice to handle non-composite root values but StackState Deserialize just drops them on the floor
+		break;
 	}
 
-	return nullptr;
+	return false;
 }
 
-void NestedVariantFromJson(const char* InJson, FNestedVariantJson::FMapRef& Map)
+bool FJsonSerializerPolicy_NestedVariant::GetValueFromState(const StackState& State, FArrayOfValues& OutArray)
 {
-	if (!InJson || strlen(InJson) == 0)
+	// Returning an empty array is ok
+	if (State.Type != EJson::Array)
 	{
-		return;
+		return false;
 	}
 
-	TSharedPtr<FJsonObject> JsonObject = MakeShared<FJsonObject>();
-	TSharedRef<TJsonReader<UTF8CHAR>> Reader = TJsonReaderFactory<UTF8CHAR>::Create(InJson);
-	verify(FJsonSerializer::Deserialize(Reader, JsonObject));
-	if (!JsonObject)
-	{
-		return;
-	}
+	OutArray = State.Array;
 
-	NestedVariantFromJsonObject(JsonObject, Map);
+	return true;
 }
 
-void NestedVariantFromJsonObject(const TSharedPtr<FJsonObject>& JsonObject, FNestedVariantJson::FMapRef& Map)
+bool FJsonSerializerPolicy_NestedVariant::GetValueFromState(const StackState& State, FMapOfValues& OutMap)
 {
-	for (const TPair<FString, TSharedPtr<FJsonValue>>& Pair : JsonObject->Values)
+	// Returning an empty object is ok
+	if (State.Type != EJson::Object)
 	{
-		switch (Pair.Value->Type)
-		{
-		case EJson::String:
-			Map->AddVariant(Pair.Key, Pair.Value->AsString());
-			break;
-		case EJson::Number:
-		{
-			int64 IntValue;
-			if (Pair.Value->TryGetNumber(IntValue))
-			{
-				Map->AddVariant(Pair.Key, IntValue);
-			}
-			else
-			{
-				double DoubleValue = 0;
-				verify(Pair.Value->TryGetNumber(DoubleValue));
-				Map->AddVariant(Pair.Key, DoubleValue);
-			}
-			break;
-		}
-		case EJson::Boolean:
-			Map->AddVariant(Pair.Key, Pair.Value->AsBool());
-			break;
-		case EJson::Array:
-		{
-			FNestedVariantJson::FArrayRef NestedArray = FNestedVariantJson::FArray::CreateVariant();
-			NestedVariantFromJsonArray(Pair.Value->AsArray(), NestedArray);
-			Map->AddVariant(Pair.Key, NestedArray);
-			break;
-		}
-		case EJson::Object:
-		{
-			FNestedVariantJson::FMapRef NestedMap = FNestedVariantJson::FMap::CreateVariant();
-			NestedVariantFromJsonObject(Pair.Value->AsObject(), NestedMap);
-			Map->AddVariant(Pair.Key, NestedMap);
-			break;
-		}
-		default:
-			checkNoEntry();
-		}
+		return false;
 	}
+
+	OutMap = State.Object;
+
+	return true;
 }
 
-void NestedVariantFromJsonArray(const TArray<TSharedPtr<FJsonValue>>& JsonArray, FNestedVariantJson::FArrayRef& Array)
+void FJsonSerializerPolicy_NestedVariant::ResetValue(FValue& OutValue)
 {
-	for (const TSharedPtr<FJsonValue>& Value : JsonArray)
-	{
-		switch (Value->Type)
-		{
-		case EJson::String:
-			Array->AddVariant(Value->AsString());
-			break;
-		case EJson::Number:
-		{
-			int64 IntValue;
-			if (Value->TryGetNumber(IntValue))
-			{
-				Array->AddVariant(IntValue);
-			}
-			else
-			{
-				double DoubleValue = 0;
-				verify(Value->TryGetNumber(DoubleValue));
-				Array->AddVariant(DoubleValue);
-			}
-			break;
-		}
-		case EJson::Boolean:
-			Array->AddVariant(Value->AsBool());
-			break;
-		case EJson::Array:
-		{
-			FNestedVariantJson::FArrayRef NestedArray = FNestedVariantJson::FArray::CreateVariant();
-			NestedVariantFromJsonArray(Value->AsArray(), NestedArray);
-			Array->AddVariant(NestedArray);
-			break;
-		}
-		case EJson::Object:
-		{
-			FNestedVariantJson::FMapRef NestedMap = FNestedVariantJson::FMap::CreateVariant();
-			NestedVariantFromJsonObject(Value->AsObject(), NestedMap);
-			Array->AddVariant(NestedMap);
-			break;
-		}
-		default:
-			checkNoEntry();
-		}
-	}
+	OutValue = FValue();
 }
 
-void LexFromString(FNestedVariantJson::FMap& OutValue, const TCHAR* InString)
+void FJsonSerializerPolicy_NestedVariant::ReadObjectStart(StackState& State)
 {
-	FNestedVariantJson::FMapRef  OutValueRef = FNestedVariantJson::FMap::CreateVariant();
-	NestedVariantFromJson(TCHAR_TO_UTF8(InString), OutValueRef);
+	State.Type = EJson::Object;
+	State.Object = FNestedVariantJson::FMap::CreateVariant();
+}
 
-	OutValue = OutValueRef.Get();
+void FJsonSerializerPolicy_NestedVariant::ReadObjectEnd(StackState& State, FValue& OutValue)
+{
+	OutValue.Set<FNestedVariantJson::FMapRef>(State.Object.ToSharedRef());
+}
+
+void FJsonSerializerPolicy_NestedVariant::ReadArrayStart(StackState& State)
+{
+	State.Type = EJson::Array;
+	State.Array = FNestedVariantJson::FArray::CreateVariant();
+}
+
+void FJsonSerializerPolicy_NestedVariant::ReadArrayEnd(StackState& State, FValue& OutValue)
+{
+	OutValue.Set<FNestedVariantJson::FArrayRef>(State.Array.ToSharedRef());
+}
+
+void FJsonSerializerPolicy_NestedVariant::ReadNull(FValue& OutValue)
+{
+}
+
+void FJsonSerializerPolicy_NestedVariant::AddValueToObject(StackState& State, const FString& Identifier, FValue& NewValue)
+{
+	State.Object->Add(Identifier, NewValue);
+}
+
+void FJsonSerializerPolicy_NestedVariant::AddValueToArray(StackState& State, FValue& NewValue)
+{
+	State.Array->Add(NewValue);
 }
\ No newline at end of file
diff --git a/Engine/Source/Runtime/CoreOnline/Public/Online/NestedVariantJson.h b/Engine/Source/Runtime/CoreOnline/Public/Online/NestedVariantJson.h
--- a/Engine/Source/Runtime/CoreOnline/Public/Online/NestedVariantJson.h
+++ b/Engine/Source/Runtime/CoreOnline/Public/Online/NestedVariantJson.h
@@ -4,6 +4,9 @@
 
 #include "Containers/UnrealString.h"
 #include "Online/NestedVariant.h"
+#include "Serialization/JsonReader.h"
+#include "Serialization/JsonTypes.h"
+#include "Serialization/JsonWriter.h"
 
 // TODO: Update the JSON module to handle nested variants of any type directly so this file is not needed.
 
@@ -13,10 +16,287 @@ class FJsonValue;
 
 COREONLINE_API FString NestedVariantToJson(const FNestedVariantJson::FMapPtr& Map);
 COREONLINE_API TSharedRef<FJsonObject> NestedVariantToJsonObject(const FNestedVariantJson::FMapRef& Map);
-COREONLINE_API TSharedPtr<FJsonValue> NestedVariantToJsonValue(const FNestedVariantJson::FValue& Value);
 
 COREONLINE_API void NestedVariantFromJson(const char* InJson, FNestedVariantJson::FMapRef& Map);
 COREONLINE_API void NestedVariantFromJsonObject(const TSharedPtr<FJsonObject>& JsonObject, FNestedVariantJson::FMapRef& Map);
-COREONLINE_API void NestedVariantFromJsonArray(const TArray<TSharedPtr<FJsonValue>>& JsonArray, FNestedVariantJson::FArrayRef& Array);
 
-COREONLINE_API void LexFromString(FNestedVariantJson::FMap& OutValue, const TCHAR* InString);
\ No newline at end of file
+COREONLINE_API void LexFromString(FNestedVariantJson::FMap& OutValue, const TCHAR* InString);
+
+struct FJsonSerializerPolicy_NestedVariant
+{
+	using FValue = FNestedVariantJson::FValue;
+	using FArrayOfValues = FNestedVariantJson::FArrayPtr;
+	using FMapOfValues = FNestedVariantJson::FMapPtr;
+
+	struct StackState
+	{
+		EJson Type;
+		FString Identifier;
+		FNestedVariantJson::FArrayPtr Array;
+		FNestedVariantJson::FMapPtr Object;
+	};
+
+	struct FElement
+	{
+		FElement( const FValue& InValue )
+			: Identifier()
+			, Value(InValue)
+		{ }
+
+		FElement( const FMapOfValues& Object )
+			: Identifier()
+		{
+			Value.Set<FNestedVariantJson::FMapRef>(StaticCastSharedRef<FNestedVariantJson::FMap>(Object->AsShared()));
+		}
+
+		FElement( const FArrayOfValues& Array )
+			: Identifier()
+		{
+			Value.Set<FNestedVariantJson::FArrayRef>(StaticCastSharedRef<FNestedVariantJson::FArray>(Array->AsShared()));
+		}
+
+		FElement( const FString& InIdentifier, const FValue& InValue )
+			: Identifier( InIdentifier )
+			, Value( InValue )
+			, bIsKeyValuePair( true )
+		{ }
+
+		FString Identifier;
+		FValue Value;
+		bool bHasBeenProcessed = false;
+		bool bIsKeyValuePair = false;
+	};
+
+	COREONLINE_API static bool GetValueFromState(const StackState& State, FValue& OutValue);
+	COREONLINE_API static bool GetValueFromState(const StackState& State, FArrayOfValues& OutArray);
+	COREONLINE_API static bool GetValueFromState(const StackState& State, FMapOfValues& OutMap);
+	COREONLINE_API static void ResetValue(FValue& OutValue);
+	COREONLINE_API static void ReadObjectStart(StackState& State);
+	COREONLINE_API static void ReadObjectEnd(StackState& State, FValue& OutValue);
+	COREONLINE_API static void ReadArrayStart(StackState& State);
+	COREONLINE_API static void ReadArrayEnd(StackState& State, FValue& OutValue);
+
+	template<class CharType>
+	static void ReadBoolean(TJsonReader<CharType>& Reader, FValue& OutValue)
+	{
+		OutValue.template Emplace<bool>(Reader.GetValueAsBoolean());
+	}
+
+	template<class CharType>
+	static void ReadString(TJsonReader<CharType>& Reader, FValue& OutValue)
+	{
+		OutValue.template Emplace<FString>(Reader.StealInternalValueAsString());
+	}
+
+	template<class CharType>
+	static void ReadNumberAsString(TJsonReader<CharType>& Reader, FValue& OutValue)
+	{
+		OutValue.template Emplace<FString>(Reader.GetValueAsNumberString());
+	}
+
+	template<class CharType>
+	static bool IsNumberStringInteger(const CharType* Str)
+	{
+		if (*Str == '-' || *Str == '+')
+		{
+			Str++;
+		}
+
+		while (*Str != '\0')
+		{
+			if (*Str == '.')
+			{
+				return false;
+			}
+			else if (!FChar::IsDigit(*Str))
+			{
+				return false;
+			}
+			
+			++Str;
+		}
+
+		return true;
+	}
+
+	template<class CharType>
+	static void ReadNumber(TJsonReader<CharType>& Reader, FValue& OutValue)
+	{
+		const typename TJsonReader<CharType>::StoredStringType& NumberString = Reader.GetValueAsNumberString();
+		if (IsNumberStringInteger(*NumberString))
+		{
+			OutValue.template Emplace<long long>(Reader.GetValueAsNumber());
+		}
+		else
+		{
+			OutValue.template Emplace<double>(Reader.GetValueAsNumber());
+		}
+	}
+
+	static void ReadNull(FValue& OutValue);
+	static void AddValueToObject(StackState& State, const FString& Identifier, FValue& NewValue);
+	static void AddValueToArray(StackState& State, FValue& NewValue);
+
+	template<class CharType, class PrintPolicy>
+	static bool SerializeIfBool(TSharedRef<FElement>& Element,  TJsonWriter<CharType, PrintPolicy>& Writer, bool bWriteValueOnly)
+	{
+		if(Element->Value.IsType<bool>())
+		{
+			if (bWriteValueOnly)
+			{
+				Writer.WriteValue(Element->Value.Get<bool>());
+			}
+			else
+			{
+				Writer.WriteValue(Element->Identifier, Element->Value.Get<bool>());
+			}
+
+			return true;
+		}
+
+		return false;
+	}
+
+	template<class CharType, class PrintPolicy>
+	static bool SerializeIfNumber(TSharedRef<FElement>& Element,  TJsonWriter<CharType, PrintPolicy>& Writer, bool bWriteValueOnly)
+	{
+		if (Element->Value.IsType<long long>())
+		{
+			if (bWriteValueOnly)
+			{
+				Writer.WriteValue(Element->Value.Get<long long>());
+			}
+			else
+			{
+				Writer.WriteValue(Element->Identifier, Element->Value.Get<long long>());
+			}
+
+			return true;
+		}
+		else if(Element->Value.IsType<double>())
+		{
+			if (bWriteValueOnly)
+			{
+				Writer.WriteValue(Element->Value.Get<double>());
+			}
+			else
+			{
+				Writer.WriteValue(Element->Identifier, Element->Value.Get<double>());
+			}
+
+			return true;
+		}
+
+		return false;
+	}
+
+	template<class CharType, class PrintPolicy>
+	static bool SerializeIfString(TSharedRef<FElement>& Element,  TJsonWriter<CharType, PrintPolicy>& Writer, bool bWriteValueOnly)
+	{
+		if(Element->Value.IsType<FString>())
+		{
+			if (bWriteValueOnly)
+			{
+				Writer.WriteValue(Element->Value.Get<FString>());
+			}
+			else
+			{
+				Writer.WriteValue(Element->Identifier, Element->Value.Get<FString>());
+			}
+
+			return true;
+		}
+
+		return false;
+	}
+
+	template<class CharType, class PrintPolicy>
+	static bool SerializeIfNull(TSharedRef<FElement>& Element,  TJsonWriter<CharType, PrintPolicy>& Writer, bool bWriteValueOnly)
+	{		
+		return false;
+	}
+
+	template<class CharType, class PrintPolicy>
+	static bool SerializeIfArray(TArray<TSharedRef<FElement>>& ElementStack, TSharedRef<FElement>& Element,  TJsonWriter<CharType, PrintPolicy>& Writer, bool bWriteValueOnly)
+	{
+		if(Element->Value.IsType<FNestedVariantJson::FArrayRef>())
+		{
+			if (Element->bHasBeenProcessed)
+			{
+				Writer.WriteArrayEnd();
+			}
+			else
+			{
+				Element->bHasBeenProcessed = true;
+				ElementStack.Push(Element);
+
+				if (bWriteValueOnly)
+				{
+					Writer.WriteArrayStart();
+				}
+				else
+				{
+					Writer.WriteArrayStart(Element->Identifier);
+				}
+
+				FArrayOfValues Values = StaticCastSharedPtr<FNestedVariantJson::FArray>(Element->Value.Get<FNestedVariantJson::FArrayRef>().ToSharedPtr());
+
+				for (int Index = Values->Num() - 1; Index >= 0; --Index)
+				{
+					ElementStack.Push(MakeShared<FElement>((*Values)[Index]));
+				}
+			}
+
+			return true;
+		}
+		
+		return false;
+	}
+
+	template<class CharType, class PrintPolicy>
+	static bool SerializeIfObject(TArray<TSharedRef<FElement>>& ElementStack, TSharedRef<FElement>& Element,  TJsonWriter<CharType, PrintPolicy>& Writer, bool bWriteValueOnly)
+	{
+		if (Element->Value.IsType<FNestedVariantJson::FMapRef>())
+		{
+			if (Element->bHasBeenProcessed)
+			{
+				Writer.WriteObjectEnd();
+			}
+			else
+			{
+				Element->bHasBeenProcessed = true;
+				ElementStack.Push(Element);
+
+				if (bWriteValueOnly)
+				{
+					Writer.WriteObjectStart();
+				}
+				else
+				{
+					Writer.WriteObjectStart(Element->Identifier);
+				}
+
+				TArray<FString> Keys; 
+				TArray<FValue> Values;
+				FMapOfValues ElementMap = StaticCastSharedPtr<FNestedVariantJson::FMap>(Element->Value.Get<FNestedVariantJson::FMapRef>().ToSharedPtr());
+				for (const FNestedVariantJson::FMap::ElementType& Entry : *ElementMap)
+				{
+					Keys.Add(Entry.Key);
+					
+					Values.Add(Entry.Value);
+				}
+
+				check(Keys.Num() == Values.Num());
+
+				for (int Index = Values.Num() - 1; Index >= 0; --Index)
+				{
+					ElementStack.Push(MakeShared<FElement>(Keys[Index], Values[Index]));
+				}
+			}
+
+			return true;
+		}
+		
+		return false;
+	}
+};
\ No newline at end of file
diff --git a/Engine/Source/Runtime/Json/Private/Serialization/JsonSerializer.cpp b/Engine/Source/Runtime/Json/Private/Serialization/JsonSerializer.cpp
new file mode 100644
--- /dev/null
+++ b/Engine/Source/Runtime/Json/Private/Serialization/JsonSerializer.cpp
@@ -0,0 +1,96 @@
+// Copyright Epic Games, Inc. All Rights Reserved.
+
+#include "Serialization/JsonSerializer.h"
+
+bool FJsonSerializerPolicy_JsonObject::GetValueFromState(const StackState& State, FValue& OutValue)
+{
+	switch (State.Type)
+	{
+	case EJson::Object:
+		if (!State.Object.IsValid())
+		{
+			return false;
+		}
+		OutValue = MakeShared<FJsonValueObject>(State.Object);
+		break;
+	case EJson::Array:
+		OutValue = MakeShared<FJsonValueArray>(State.Array);
+		break;
+	case EJson::None:
+	case EJson::Null:
+	case EJson::String:
+	case EJson::Number:
+	case EJson::Boolean:
+	default:
+		// FIXME: would be nice to handle non-composite root values but StackState Deserialize just drops them on the floor
+		return false;
+	}
+	return true;
+}
+
+bool FJsonSerializerPolicy_JsonObject::GetValueFromState(const StackState& State, FArrayOfValues& OutArray)
+{
+	// Returning an empty array is ok.
+	if (State.Type != EJson::Array)
+	{
+		return false;
+	}
+
+	OutArray = State.Array;
+
+	return true;
+}
+
+bool FJsonSerializerPolicy_JsonObject::GetValueFromState(const StackState& State, FMapOfValues& OutMap)
+{
+	// Returning an empty object is ok
+	if (!State.Object.IsValid())
+	{
+		return false;
+	}
+
+	OutMap = State.Object;
+
+	return true;
+}
+
+void FJsonSerializerPolicy_JsonObject::ResetValue(FValue& OutValue)
+{
+	OutValue.Reset();
+}
+
+void FJsonSerializerPolicy_JsonObject::ReadObjectStart(StackState& State)
+{
+	State.Type = EJson::Object;
+	State.Object = MakeShared<FJsonObject>();
+}
+
+void FJsonSerializerPolicy_JsonObject::ReadObjectEnd(StackState& State, FValue& OutValue)
+{
+	OutValue = MakeShared<FJsonValueObject>(State.Object);
+}
+
+void FJsonSerializerPolicy_JsonObject::ReadArrayStart(StackState& State)
+{
+	State.Type = EJson::Array;
+}
+
+void FJsonSerializerPolicy_JsonObject::ReadArrayEnd(StackState& State, FValue& OutValue)
+{
+	OutValue = MakeShared<FJsonValueArray>(State.Array);
+}
+
+void FJsonSerializerPolicy_JsonObject::ReadNull(FValue& OutValue)
+{
+	OutValue = MakeShared<FJsonValueNull>();
+}
+
+void FJsonSerializerPolicy_JsonObject::AddValueToObject(StackState& State, const FString& Identifier, FValue& NewValue)
+{
+	State.Object->SetField(Identifier, NewValue);
+}
+
+void FJsonSerializerPolicy_JsonObject::AddValueToArray(StackState& State, FValue& NewValue)
+{
+	State.Array.Add(NewValue);
+}
\ No newline at end of file
diff --git a/Engine/Source/Runtime/Json/Public/Serialization/JsonSerializer.h b/Engine/Source/Runtime/Json/Public/Serialization/JsonSerializer.h
--- a/Engine/Source/Runtime/Json/Public/Serialization/JsonSerializer.h
+++ b/Engine/Source/Runtime/Json/Public/Serialization/JsonSerializer.h
@@ -11,294 +11,484 @@
 
 class Error;
 
-class FJsonSerializer
+struct FJsonSerializerPolicy_JsonObject
 {
-public:
+	using FValue = TSharedPtr<FJsonValue>;
+	using FArrayOfValues = TArray<TSharedPtr<FJsonValue>>;
+	using FMapOfValues = TSharedPtr<FJsonObject>;
 
-	enum class EFlags
+	struct StackState
 	{
-		None = 0,
-		StoreNumbersAsStrings = 1,
+		EJson Type;
+		FString Identifier;
+		TArray<TSharedPtr<FJsonValue>> Array;
+		TSharedPtr<FJsonObject> Object;
 	};
 
-	template <class CharType>
-	static bool Deserialize(const TSharedRef<TJsonReader<CharType>>& Reader, TArray<TSharedPtr<FJsonValue>>& OutArray, EFlags InOptions = EFlags::None)
+	struct FElement
 	{
-		return Deserialize(*Reader, OutArray, InOptions);
+		FElement( const FValue& InValue )
+			: Identifier()
+			, Value(InValue)
+		{ }
+
+		FElement( const FMapOfValues& Object )
+			: Identifier()
+			, Value(MakeShared<FJsonValueObject>(Object))
+		{ }
+
+		FElement( const FArrayOfValues& Array )
+			: Identifier()
+			, Value(MakeShared<FJsonValueArray>(Array))
+		{ }
+
+		FElement( const FString& InIdentifier, const FValue& InValue )
+			: Identifier( InIdentifier )
+			, Value( InValue )
+			, bIsKeyValuePair( true )
+		{ }
+
+		FString Identifier;
+		FValue Value;
+		bool bHasBeenProcessed = false;
+		bool bIsKeyValuePair = false;
+	};
+
+	JSON_API static bool GetValueFromState(const StackState& State, FValue& OutValue);
+	JSON_API static bool GetValueFromState(const StackState& State, FArrayOfValues& OutArray);
+	JSON_API static bool GetValueFromState(const StackState& State, FMapOfValues& OutMap);
+	JSON_API static void ResetValue(FValue& OutValue);
+	JSON_API static void ReadObjectStart(StackState& State);
+	JSON_API static void ReadObjectEnd(StackState& State, FValue& OutValue);
+	JSON_API static void ReadArrayStart(StackState& State);
+	JSON_API static void ReadArrayEnd(StackState& State, FValue& OutValue);
+
+	template<class CharType>
+	static void ReadBoolean(TJsonReader<CharType>& Reader, FValue& OutValue)
+	{
+		OutValue = MakeShared<FJsonValueBoolean>(Reader.GetValueAsBoolean());
 	}
 
-	template <class CharType>
-	static bool Deserialize(TJsonReader<CharType>& Reader, TArray<TSharedPtr<FJsonValue>>& OutArray, EFlags InOptions = EFlags::None)
+	template<class CharType>
+	static void ReadString(TJsonReader<CharType>& Reader, FValue& OutValue)
 	{
-		StackState State;
-		if (!Deserialize(Reader, /*OUT*/State, InOptions))
+		using StoredCharType = typename TJsonReader<CharType>::StoredCharType;
+		OutValue = MakeShared<TJsonValueString<StoredCharType>>(Reader.StealInternalValueAsString());
+	}
+
+	template<class CharType>
+	static void ReadNumberAsString(TJsonReader<CharType>& Reader, FValue& OutValue)
+	{
+		using StoredCharType = typename TJsonReader<CharType>::StoredCharType;
+		OutValue = MakeShared<TJsonValueNumberString<StoredCharType>>(Reader.GetValueAsNumberString());
+	}
+
+	template<class CharType>
+	static void ReadNumber(TJsonReader<CharType>& Reader, FValue& OutValue)
+	{
+		OutValue = MakeShared<FJsonValueNumber>(Reader.GetValueAsNumber());
+	}
+
+	JSON_API static void ReadNull(FValue& OutValue);
+	JSON_API static void AddValueToObject(StackState& State, const FString& Identifier, FValue& NewValue);
+	JSON_API static void AddValueToArray(StackState& State, FValue& NewValue);
+
+	template<class CharType, class PrintPolicy>
+	static bool SerializeIfBool(TSharedRef<FElement>& Element,  TJsonWriter<CharType, PrintPolicy>& Writer, bool bWriteValueOnly)
+	{
+		if (Element->Value->Type == EJson::Boolean)
 		{
-			return false;
+			if (bWriteValueOnly)
+			{
+				Writer.WriteValue(Element->Value->AsBool());
+			}
+			else
+			{
+				Writer.WriteValue(Element->Identifier, Element->Value->AsBool());
+			}
+
+			return true;
 		}
 
-		// Empty array is ok.
-		if (State.Type != EJson::Array)
+		return false;
+	}
+
+	template<class CharType, class PrintPolicy>
+	static bool SerializeIfNumber(TSharedRef<FElement>& Element,  TJsonWriter<CharType, PrintPolicy>& Writer, bool bWriteValueOnly)
+	{
+		if (Element->Value->Type == EJson::Number)
 		{
-			return false;
+			if (bWriteValueOnly)
+			{
+				if (Element->Value->PreferStringRepresentation())
+				{
+					Writer.WriteRawJSONValue(Element->Value->AsString());
+				}
+				else
+				{
+					Writer.WriteValue(Element->Value->AsNumber());
+				}
+			}
+			else
+			{
+				if (Element->Value->PreferStringRepresentation())
+				{
+					Writer.WriteRawJSONValue(Element->Identifier, Element->Value->AsString());
+				}
+				else
+				{
+					Writer.WriteValue(Element->Identifier, Element->Value->AsNumber());
+				}
+			}
+
+			return true;
 		}
 
-		OutArray = State.Array;
+		return false;
+	}
 
-		return true;
+	template<class CharType, class PrintPolicy>
+	static bool SerializeIfString(TSharedRef<FElement>& Element,  TJsonWriter<CharType, PrintPolicy>& Writer, bool bWriteValueOnly)
+	{
+		if(Element->Value->Type == EJson::String)
+		{
+			if (bWriteValueOnly)
+			{
+				Writer.WriteValue(Element->Value->AsString());
+			}
+			else
+			{
+				Writer.WriteValue(Element->Identifier, Element->Value->AsString());
+			}
+
+			return true;
+		}
+
+		return false;
 	}
 
+	template<class CharType, class PrintPolicy>
+	static bool SerializeIfNull(TSharedRef<FElement>& Element,  TJsonWriter<CharType, PrintPolicy>& Writer, bool bWriteValueOnly)
+	{
+		if(Element->Value->Type == EJson::Null)
+		{
+			if (bWriteValueOnly)
+			{
+				Writer.WriteNull();
+			}
+			else
+			{
+				Writer.WriteNull(Element->Identifier);
+			}
+			
+			return true;
+		}
+		
+		return false;
+	}
+
+	template<class CharType, class PrintPolicy>
+	static bool SerializeIfArray(TArray<TSharedRef<FElement>>& ElementStack, TSharedRef<FElement>& Element,  TJsonWriter<CharType, PrintPolicy>& Writer, bool bWriteValueOnly)
+	{
+		if(Element->Value->Type == EJson::Array)
+		{
+			if (Element->bHasBeenProcessed)
+			{
+				Writer.WriteArrayEnd();
+			}
+			else
+			{
+				Element->bHasBeenProcessed = true;
+				ElementStack.Push(Element);
+
+				if (bWriteValueOnly)
+				{
+					Writer.WriteArrayStart();
+				}
+				else
+				{
+					Writer.WriteArrayStart(Element->Identifier);
+				}
+
+				FArrayOfValues Values = Element->Value->AsArray();
+
+				for (int Index = Values.Num() - 1; Index >= 0; --Index)
+				{
+					ElementStack.Push(MakeShared<FElement>(Values[Index]));
+				}
+			}
+
+			return true;
+		}
+		
+		return false;
+	}
+
+	template<class CharType, class PrintPolicy>
+	static bool SerializeIfObject(TArray<TSharedRef<FElement>>& ElementStack, TSharedRef<FElement>& Element,  TJsonWriter<CharType, PrintPolicy>& Writer, bool bWriteValueOnly)
+	{
+		if(Element->Value->Type == EJson::Object)
+		{
+			if (Element->bHasBeenProcessed)
+			{
+				Writer.WriteObjectEnd();
+			}
+			else
+			{
+				Element->bHasBeenProcessed = true;
+				ElementStack.Push(Element);
+
+				if (bWriteValueOnly)
+				{
+					Writer.WriteObjectStart();
+				}
+				else
+				{
+					Writer.WriteObjectStart(Element->Identifier);
+				}
+
+				TArray<FString> Keys; 
+				FArrayOfValues Values;
+				FMapOfValues ElementObject = Element->Value->AsObject();
+				ElementObject->Values.GenerateKeyArray(Keys);
+				ElementObject->Values.GenerateValueArray(Values);
+
+				check(Keys.Num() == Values.Num());
+
+				for (int Index = Values.Num() - 1; Index >= 0; --Index)
+				{
+					ElementStack.Push(MakeShared<FElement>(Keys[Index], Values[Index]));
+				}
+			}
+
+			return true;
+		}
+		
+		return false;
+	}
+};
+
+template<typename Policy = FJsonSerializerPolicy_JsonObject>
+class TJsonSerializer
+{
+public:
+
+	enum class EFlags
+	{
+		None = 0,
+		StoreNumbersAsStrings = 1,
+	};
+
 	template <class CharType>
-	static bool Deserialize(const TSharedRef<TJsonReader<CharType>>& Reader, TSharedPtr<FJsonObject>& OutObject, EFlags InOptions = EFlags::None)
+	static bool Deserialize(const TSharedRef<TJsonReader<CharType>>& Reader, Policy::FValue& OutValue, EFlags InOptions = EFlags::None)
 	{
-		return Deserialize(*Reader, OutObject, InOptions);
+		return Deserialize(*Reader, OutValue, InOptions);
 	}
+
 	template <class CharType>
-	static bool Deserialize(TJsonReader<CharType>& Reader, TSharedPtr<FJsonObject>& OutObject, EFlags InOptions = EFlags::None)
+	static bool Deserialize(TJsonReader<CharType>& Reader, Policy::FValue& OutValue, EFlags InOptions = EFlags::None)
 	{
-		StackState State;
+		typename Policy::StackState State;
 		if (!Deserialize(Reader, /*OUT*/State, InOptions))
 		{
 			return false;
 		}
 
-		if (!State.Object.IsValid())
+		return Policy::GetValueFromState(State, OutValue);
+	}
+
+	template <class CharType>
+	static bool Deserialize(const TSharedRef<TJsonReader<CharType>>& Reader, Policy::FArrayOfValues& OutArray, EFlags InOptions = EFlags::None)
+	{
+		return Deserialize(*Reader, OutArray, InOptions);
+	}
+
+	template <class CharType>
+	static bool Deserialize(TJsonReader<CharType>& Reader, Policy::FArrayOfValues& OutArray, EFlags InOptions = EFlags::None)
+	{
+		typename Policy::StackState State;
+		if (!Deserialize(Reader, /*OUT*/State, InOptions))
 		{
 			return false;
 		}
 
-		OutObject = State.Object;
-
-		return true;
+		return Policy::GetValueFromState(State, OutArray);
 	}
 
 	template <class CharType>
-	static bool Deserialize(const TSharedRef<TJsonReader<CharType>>& Reader, TSharedPtr<FJsonValue>& OutValue, EFlags InOptions = EFlags::None)
+	static bool Deserialize(const TSharedRef<TJsonReader<CharType>>& Reader, Policy::FMapOfValues& OutMap, EFlags InOptions = EFlags::None)
 	{
-		return Deserialize(*Reader, OutValue, InOptions);
+		return Deserialize(*Reader, OutMap, InOptions);
 	}
 
 	template <class CharType>
-	static bool Deserialize(TJsonReader<CharType>& Reader, TSharedPtr<FJsonValue>& OutValue, EFlags InOptions = EFlags::None)
+	static bool Deserialize(TJsonReader<CharType>& Reader, Policy::FMapOfValues& OutMap, EFlags InOptions = EFlags::None)
 	{
-		StackState State;
+		typename Policy::StackState State;
 		if (!Deserialize(Reader, /*OUT*/State, InOptions))
 		{
 			return false;
 		}
 
-		switch (State.Type)
-		{
-		case EJson::Object:
-			if (!State.Object.IsValid())
-			{
-				return false;
-			}
-			OutValue = MakeShared<FJsonValueObject>(State.Object);
-			break;
-		case EJson::Array:
-			OutValue = MakeShared<FJsonValueArray>(State.Array);
-			break;
-		case EJson::None:
-		case EJson::Null:
-		case EJson::String:
-		case EJson::Number:
-		case EJson::Boolean:
-		default:
-			// FIXME: would be nice to handle non-composite root values but StackState Deserialize just drops them on the floor
-			return false;
-		}
-		return true;
+		return Policy::GetValueFromState(State, OutMap);
 	}
 
 	/**
-	 * Serialize the passed array of json values into the writer.
-	 * This will effectively serialize all of the values enclosed in [] square brackets.
-	 * Example:
-	 *  - Writer state: [123 <writer position>
-	 *    Parameter: Array: ["foo", "bar", "", 456]
-	 *    Serialization result: [123, ["foo", "bar", "", 456] <writer position>
+	 * Serialize the passed value and identifier into the writer.
+	 * Empty string identifiers will be ignored when the writer is not writing inside of a map of values and only the value will be serialized.
+	 * If the writer is in a state where it's currently writing inside of a map of values, then the identifier will always be serialized.
+	 * 
+	 * Json Examples:
+	 *	- Writer state: { "foo": "bar" <writer position>
+	 *    Parameters: Identifier: ""
+	 *                Value: "baz"
+	 *    Serialization result: { "foo": "bar", "": "baz" <writer position> //empty identifier is serialized as a valid key for the key:value pair "":"baz"
 	 *
-	 * @param Array		    The json array we are serializing
-	 * @param Writer		The writer the array is written into.
+	 * - Writer state: { "foo": ["bar" <writer position>
+	 *   Parameters: Identifier: ""
+	 *               Value: "baz"
+	 *   Serialization result: { foo: ["bar", "baz" <writer position> //empty identifier is ignored since we are writing into an array and not an object.
+	 *
+	 * @param Value			The value we are serializing
+	 * @param Identifier	The identifier of the value, empty identifiers are ignored outside of maps of values.
+	 * @param Writer		The writer the value and identifier are written into.
 	 * @param bCloseWriter	When set to true the Writer will be closed after the serialization.
 	 * @return				Returns true if the serialization was successful, false otherwise.
 	 */
 	template <class CharType, class PrintPolicy>
-	static bool Serialize(const TArray<TSharedPtr<FJsonValue>>& Array, const TSharedRef<TJsonWriter<CharType, PrintPolicy>>& Writer, bool bCloseWriter = true)
+	static bool Serialize(const Policy::FValue& Value, const FString& Identifier, const TSharedRef<TJsonWriter<CharType, PrintPolicy>>& Writer, bool bCloseWriter = true)
 	{
-		return Serialize(Array, *Writer, bCloseWriter);
+		return Serialize(Value, Identifier, *Writer, bCloseWriter);
 	}
 
 	/**
-	 * Serialize the passed array of json values into the writer.
-	 * This will effectively serialize all of the values enclosed in [] square brackets.
-	 * Example:
-	 *  - Writer state: [123 <writer position>
-	 *    Parameter: Array: ["foo", "bar", "", 456]
-	 *    Serialization result: [123, ["foo", "bar", "", 456] <writer position>
+	 * Serialize the passed value and identifier into the writer.
+	 * Empty string identifiers will be ignored when the writer is not writing inside of a map of values and only the value will be serialized.
+	 * If the writer is in a state where it's currently writing inside of a map of values, then the identifier will always be serialized.
+	 * 
+	 * Json Examples:
+	 *	- Writer state: { "foo": "bar" <writer position>
+	 *    Parameters: Identifier: ""
+	 *                Value: "baz"
+	 *    Serialization result: { "foo": "bar", "": "baz" <writer position> //empty identifier is serialized as a valid key for the key:value pair "":"baz"
 	 *
-	 * @param Array		    The json array we are serializing
-	 * @param Writer		The writer the array is written into.
+	 * - Writer state: { "foo": ["bar" <writer position>
+	 *   Parameters: Identifier: ""
+	 *               Value: "baz"
+	 *   Serialization result: { foo: ["bar", "baz" <writer position> //empty identifier is ignored since we are writing into an array and not an object.
+	 *
+	 * @param Value			The value we are serializing
+	 * @param Identifier	The identifier of the value, empty identifiers are ignored outside of maps of values.
+	 * @param Writer		The writer the value and identifier are written into.
 	 * @param bCloseWriter	When set to true the Writer will be closed after the serialization.
 	 * @return				Returns true if the serialization was successful, false otherwise.
 	 */
 	template <class CharType, class PrintPolicy>
-	static bool Serialize(const TArray<TSharedPtr<FJsonValue>>& Array, TJsonWriter<CharType, PrintPolicy>& Writer, bool bCloseWriter = true )
+	static bool Serialize(const Policy::FValue& Value, const FString& Identifier, TJsonWriter<CharType, PrintPolicy>& Writer, bool bCloseWriter = true)
 	{
-		const TSharedRef<FElement> StartingElement = MakeShared<FElement>(Array);
-		return FJsonSerializer::Serialize<CharType, PrintPolicy>(StartingElement, Writer, bCloseWriter);
+		const TSharedRef<typename Policy::FElement> StartingElement = MakeShared<typename Policy::FElement>(Identifier, Value);
+		return Serialize<CharType, PrintPolicy>(StartingElement, Writer, bCloseWriter);
 	}
 
 	/**
-	 * Serialize the passed Json object into the writer.
-	 * This will effectively serialize all of the identifier:value pairs of the object enclosed in {} curly brackets.
-	 * Example:
+	 * Serialize the passed array of values into the writer.
+	 * This will effectively serialize all of the values enclosed in [] square brackets.
+	 * 
+	 * Json Example:
 	 *  - Writer state: [123 <writer position>
-	 *    Parameter: Object: {"foo": "bar", "baz": "", "": 456}
-	 *    Serialization result: [123, {"foo": "bar", "baz": "", "": 456} <writer position>
+	 *    Parameter: Array: ["foo", "bar", "", 456]
+	 *    Serialization result: [123, ["foo", "bar", "", 456] <writer position>
 	 *
-	 * @param Object		The json object we are serializing
-	 * @param Writer		The writer the object is written into.
+	 * @param Array		    The array we are serializing
+	 * @param Writer		The writer the array is written into.
 	 * @param bCloseWriter	When set to true the Writer will be closed after the serialization.
 	 * @return				Returns true if the serialization was successful, false otherwise.
 	 */
 	template <class CharType, class PrintPolicy>
-	static bool Serialize(const TSharedRef<FJsonObject>& Object, const TSharedRef<TJsonWriter<CharType, PrintPolicy>>& Writer, bool bCloseWriter = true )
+	static bool Serialize(const Policy::FArrayOfValues& Array, const TSharedRef<TJsonWriter<CharType, PrintPolicy>>& Writer, bool bCloseWriter = true)
 	{
-		return Serialize(Object, *Writer, bCloseWriter);
+		return Serialize(Array, *Writer, bCloseWriter);
 	}
 
 	/**
-	 * Serialize the passed Json object into the writer.
-	 * This will effectively serialize all of the identifier:value pairs of the object enclosed in {} brackets.
-	 * Example:
-	 *  - Writer state: [123 <writer position>
-	 *    Parameter: Object: {"foo": "bar", "baz": "", "": 456}
-	 *    Serialization result: [123, {"foo": "bar", "baz": "", "": 456} <writer position>
+	 * Serialize the passed array of values into the writer.
+	 * This will effectively serialize all of the values enclosed in [] square brackets.
 	 * 
-     * @param Object		The json object we are serializing
-	 * @param Writer		The writer the object is written into.
+	 * Json Example:
+	 *  - Writer state: [123 <writer position>
+	 *    Parameter: Array: ["foo", "bar", "", 456]
+	 *    Serialization result: [123, ["foo", "bar", "", 456] <writer position>
+	 *
+	 * @param Array		    The array we are serializing
+	 * @param Writer		The writer the array is written into.
 	 * @param bCloseWriter	When set to true the Writer will be closed after the serialization.
 	 * @return				Returns true if the serialization was successful, false otherwise.
 	 */
 	template <class CharType, class PrintPolicy>
-	static bool Serialize(const TSharedRef<FJsonObject>& Object, TJsonWriter<CharType, PrintPolicy>& Writer, bool bCloseWriter = true)
+	static bool Serialize(const Policy::FArrayOfValues& Array, TJsonWriter<CharType, PrintPolicy>& Writer, bool bCloseWriter = true )
 	{
-		const TSharedRef<FElement> StartingElement = MakeShared<FElement>(Object);
-		return FJsonSerializer::Serialize<CharType, PrintPolicy>(StartingElement, Writer, bCloseWriter);
+		const TSharedRef<typename Policy::FElement> StartingElement = MakeShared<typename Policy::FElement>(Array);
+		return Serialize<CharType, PrintPolicy>(StartingElement, Writer, bCloseWriter);
 	}
 
 	/**
-	 * Serialize the passed Json value and identifier into the writer.
-	 * Empty string identifiers will be ignored when the writer is not writing inside of a json object and only the value will be serialized.
-	 * If the writer is in a state where it's currently writing inside of a json object, then the identifier will always be serialized.
-	 * Examples:
-	 *	- Writer state: { "foo": "bar" <writer position>
-	 *    Parameters: Identifier: ""
-	 *                Value: "baz"
-	 *    Serialization result: { "foo": "bar", "": "baz" <writer position> //empty identifier is serialized as a valid key for the key:value pair "":"baz"
-	 *
-	 * - Writer state: { "foo": ["bar" <writer position>
-	 *   Parameters: Identifier: ""
-	 *               Value: "baz"
-	 *   Serialization result: { foo: ["bar", "baz" <writer position> //empty identifier is ignored since we are writing into an array and not an object.
+	 * Serialize the passed map of values into the writer.
+	 * This will effectively serialize all of the identifier:value pairs of the map, enclosed in {} curly brackets.
+	 * 
+	 * Json Example:
+	 *  - Writer state: [123 <writer position>
+	 *    Parameter: Object: {"foo": "bar", "baz": "", "": 456}
+	 *    Serialization result: [123, {"foo": "bar", "baz": "", "": 456} <writer position>
 	 *
-	 * @param Value			The json value we are serializing
-	 * @param Identifier	The identifier of the value, empty identifiers are ignored outside of json objects.
-	 * @param Writer		The writer the value and identifier are written into.
+	 * @param Object		The map of values we are serializing
+	 * @param Writer		The writer the object is written into.
 	 * @param bCloseWriter	When set to true the Writer will be closed after the serialization.
 	 * @return				Returns true if the serialization was successful, false otherwise.
 	 */
 	template <class CharType, class PrintPolicy>
-	static bool Serialize(const TSharedPtr<FJsonValue>& Value, const FString& Identifier, const TSharedRef<TJsonWriter<CharType, PrintPolicy>>& Writer, bool bCloseWriter = true)
+	static bool Serialize(const Policy::FMapOfValues& Object, const TSharedRef<TJsonWriter<CharType, PrintPolicy>>& Writer, bool bCloseWriter = true )
 	{
-		return Serialize(Value, Identifier, *Writer, bCloseWriter);
+		return Serialize(Object, *Writer, bCloseWriter);
 	}
 
 	/**
-	 * Serialize the passed Json value and identifier into the writer.
-	 * Empty string identifiers will be ignored when the writer is not writing inside of a json object and only the value will be serialized.
-	 * If the writer is in a state where it's currently writing inside of a json object, then the identifier will always be serialized.
-	 * Examples:
-	 *	- Writer state: { "foo": "bar" <writer position>
-	 *    Parameters: Identifier: ""
-	 *                Value: "baz"
-	 *    Serialization result: { "foo": "bar", "": "baz" <writer position> //empty identifier is serialized as a valid key for the key:value pair "":"baz"
+	 * Serialize the passed map of values into the writer.
+	 * This will effectively serialize all of the identifier:value pairs of the map, enclosed in {} curly brackets.
 	 * 
-	 * - Writer state: { "foo": ["bar" <writer position>
-	 *   Parameters: Identifier: ""
-	 *               Value: "baz"
-	 *   Serialization result: { foo: ["bar", "baz" <writer position> //empty identifier is ignored since we are writing into an array and not an object.
-	 * 
-	 * @param Value			The json value we are serializing
-	 * @param Identifier	The identifier of the value, empty identifiers are ignored outside of json objects.
-	 * @param Writer		The writer the value and identifier are written into.
+	 * Json Example:
+	 *  - Writer state: [123 <writer position>
+	 *    Parameter: Object: {"foo": "bar", "baz": "", "": 456}
+	 *    Serialization result: [123, {"foo": "bar", "baz": "", "": 456} <writer position>
+	 *
+	 * @param Object		The map of values we are serializing
+	 * @param Writer		The writer the object is written into.
 	 * @param bCloseWriter	When set to true the Writer will be closed after the serialization.
 	 * @return				Returns true if the serialization was successful, false otherwise.
 	 */
 	template <class CharType, class PrintPolicy>
-	static bool Serialize(const TSharedPtr<FJsonValue>& Value, const FString& Identifier, TJsonWriter<CharType, PrintPolicy>& Writer, bool bCloseWriter = true)
+	static bool Serialize(const Policy::FMapOfValues& Object, TJsonWriter<CharType, PrintPolicy>& Writer, bool bCloseWriter = true)
 	{
-		const TSharedRef<FElement> StartingElement = MakeShared<FElement>(Identifier, Value);
-		return FJsonSerializer::Serialize<CharType, PrintPolicy>(StartingElement, Writer, bCloseWriter);
+		const TSharedRef<typename Policy::FElement> StartingElement = MakeShared<typename Policy::FElement>(Object);
+		return Serialize<CharType, PrintPolicy>(StartingElement, Writer, bCloseWriter);
 	}
 
 private:
-
-	struct StackState
-	{
-		EJson Type;
-		FString Identifier;
-		TArray<TSharedPtr<FJsonValue>> Array;
-		TSharedPtr<FJsonObject> Object;
-	};
-
-	struct FElement
-	{
-		FElement( const TSharedPtr<FJsonValue>& InValue )
-			: Identifier()
-			, Value(InValue)
-		{ }
-
-		FElement( const TSharedRef<FJsonObject>& Object )
-			: Identifier()
-			, Value(MakeShared<FJsonValueObject>(Object))
-		{ }
-
-		FElement( const TArray<TSharedPtr<FJsonValue>>& Array )
-			: Identifier()
-			, Value(MakeShared<FJsonValueArray>(Array))
-		{ }
-
-		FElement( const FString& InIdentifier, const TSharedPtr< FJsonValue >& InValue )
-			: Identifier( InIdentifier )
-			, Value( InValue )
-			, bIsKeyValuePair( true )
-		{
-
-		}
-
-		FString Identifier;
-		TSharedPtr< FJsonValue > Value;
-		bool bHasBeenProcessed = false;
-		bool bIsKeyValuePair = false;
-	};
-
-private:
-
 	template <class CharType>
-	static bool Deserialize(TJsonReader<CharType>& Reader, StackState& OutStackState, EFlags InOptions)
+	static bool Deserialize(TJsonReader<CharType>& Reader, Policy::StackState& OutStackState, EFlags InOptions)
 	{
-		TArray<TSharedRef<StackState>> ScopeStack; 
-		TSharedPtr<StackState> CurrentState;
+		TArray<TSharedRef<typename Policy::StackState>> ScopeStack; 
+		TSharedPtr<typename Policy::StackState> CurrentState;
 
-		TSharedPtr<FJsonValue> NewValue;
+		typename Policy::FValue NewValue;
 		EJsonNotation Notation;
 
 		while (Reader.ReadNext(Notation))
 		{
 			FString Identifier = Reader.GetIdentifier();
-			NewValue.Reset();
+
+			Policy::ResetValue(NewValue);
+
+			bool bWasValueRead = false;
 
 			switch( Notation )
 			{
@@ -309,10 +499,9 @@ private:
 						ScopeStack.Push(CurrentState.ToSharedRef());
 					}
 
-					CurrentState = MakeShared<StackState>();
-					CurrentState->Type = EJson::Object;
+					CurrentState = MakeShared<typename Policy::StackState>();
 					CurrentState->Identifier = Identifier;
-					CurrentState->Object = MakeShared<FJsonObject>();
+					Policy::ReadObjectStart(*CurrentState);
 				}
 				break;
 
@@ -321,7 +510,8 @@ private:
 					if (ScopeStack.Num() > 0)
 					{
 						Identifier = CurrentState->Identifier;
-						NewValue = MakeShared<FJsonValueObject>(CurrentState->Object);
+						Policy::ReadObjectEnd(*CurrentState, NewValue);
+						bWasValueRead = true;
 						CurrentState = ScopeStack.Pop();
 					}
 				}
@@ -334,8 +524,8 @@ private:
 						ScopeStack.Push(CurrentState.ToSharedRef());
 					}
 
-					CurrentState = MakeShared<StackState>();
-					CurrentState->Type = EJson::Array;
+					CurrentState = MakeShared<typename Policy::StackState>();
+					Policy::ReadArrayStart(*CurrentState);
 					CurrentState->Identifier = Identifier;
 				}
 				break;
@@ -345,37 +535,38 @@ private:
 					if (ScopeStack.Num() > 0)
 					{
 						Identifier = CurrentState->Identifier;
-						NewValue = MakeShared<FJsonValueArray>(CurrentState->Array);
+						Policy::ReadArrayEnd(*CurrentState, NewValue);
+						bWasValueRead = true;
 						CurrentState = ScopeStack.Pop();
 					}
 				}
 				break;
 
 			case EJsonNotation::Boolean:
-				NewValue = MakeShared<FJsonValueBoolean>(Reader.GetValueAsBoolean());
+				Policy::ReadBoolean(Reader, NewValue); bWasValueRead = true;
 				break;
 
 			case EJsonNotation::String:
-				{
-					using StoredCharType = typename TJsonReader<CharType>::StoredCharType;
-					NewValue = MakeShared<TJsonValueString<StoredCharType>>(Reader.StealInternalValueAsString());
-				}
+				Policy::ReadString(Reader, NewValue); bWasValueRead = true;
 				break;
 
 			case EJsonNotation::Number:
-				if (EnumHasAnyFlags(InOptions, EFlags::StoreNumbersAsStrings))
-				{
-					using StoredCharType = typename TJsonReader<CharType>::StoredCharType;
-					NewValue = MakeShared<TJsonValueNumberString<StoredCharType>>(Reader.GetValueAsNumberString());
-				}
-				else
 				{
-					NewValue = MakeShared<FJsonValueNumber>(Reader.GetValueAsNumber());
+					if (EnumHasAnyFlags(InOptions, EFlags::StoreNumbersAsStrings))
+					{
+						Policy::ReadNumberAsString(Reader, NewValue);
+					}
+					else
+					{
+						Policy::ReadNumber(Reader, NewValue);
+					}
+
+					bWasValueRead = true;
 				}
 				break;
 
 			case EJsonNotation::Null:
-				NewValue = MakeShared<FJsonValueNull>();
+				Policy::ReadNull(NewValue); bWasValueRead = true;
 				break;
 
 			case EJsonNotation::Error:
@@ -383,15 +574,15 @@ private:
 				break;
 			}
 
-			if (NewValue.IsValid() && CurrentState.IsValid())
+			if (bWasValueRead && CurrentState.IsValid())
 			{
 				if (CurrentState->Type == EJson::Object)
 				{
-					CurrentState->Object->SetField(Identifier, NewValue);
+					Policy::AddValueToObject(*CurrentState, Identifier, NewValue);
 				}
 				else
 				{
-					CurrentState->Array.Add(NewValue);
+					Policy::AddValueToArray(*CurrentState, NewValue);
 				}
 			}
 		}
@@ -404,159 +595,52 @@ private:
 		OutStackState = *CurrentState.Get();
 
 		return true;
-	}
+	}	
 
 	template <class CharType, class PrintPolicy>
-	static bool Serialize(const TSharedRef<FElement>& StartingElement, TJsonWriter<CharType, PrintPolicy>& Writer, bool bCloseWriter)
+	static bool Serialize(const TSharedRef<typename Policy::FElement>& StartingElement, TJsonWriter<CharType, PrintPolicy>& Writer, bool bCloseWriter)
 	{
-		TArray<TSharedRef<FElement>> ElementStack;
+		TArray<TSharedRef<typename Policy::FElement>> ElementStack;
 		ElementStack.Push(StartingElement);
 
 		while (ElementStack.Num() > 0)
 		{
-			TSharedRef<FElement> Element = ElementStack.Pop();
+			TSharedRef<typename Policy::FElement> Element = ElementStack.Pop();			
+
 			// Empty keys are valid identifiers only when writing inside an object.
 			const bool bWriteValueOnly = !Element->bIsKeyValuePair || (Element->Identifier.IsEmpty() && Writer.GetCurrentElementType() != EJson::Object);
-			check(Element->Value->Type != EJson::None);
 
-			switch (Element->Value->Type)
+			if (Policy::SerializeIfBool(Element, Writer, bWriteValueOnly))
 			{
-			case EJson::Number:	
-				{
-					if (bWriteValueOnly)
-					{
-						if (Element->Value->PreferStringRepresentation())
-						{
-							Writer.WriteRawJSONValue(Element->Value->AsString());
-						}
-						else
-						{
-							Writer.WriteValue(Element->Value->AsNumber());
-						}
-					}
-					else
-					{
-						if (Element->Value->PreferStringRepresentation())
-						{
-							Writer.WriteRawJSONValue(Element->Identifier, Element->Value->AsString());
-						}
-						else
-						{
-							Writer.WriteValue(Element->Identifier, Element->Value->AsNumber());
-						}
-					}
-				}
-				break;
-
-			case EJson::Boolean:					
-				{
-					if (bWriteValueOnly)
-					{
-						Writer.WriteValue(Element->Value->AsBool());
-					}
-					else
-					{
-						Writer.WriteValue(Element->Identifier, Element->Value->AsBool());
-					}
-				}
-				break;
-
-			case EJson::String:
-				{
-					if (bWriteValueOnly)
-					{
-						Writer.WriteValue(Element->Value->AsString());
-					}
-					else
-					{
-						Writer.WriteValue(Element->Identifier, Element->Value->AsString());
-					}
-				}
-				break;
-
-			case EJson::Null:
-				{
-					if (bWriteValueOnly)
-					{
-						Writer.WriteNull();
-					}
-					else
-					{
-						Writer.WriteNull(Element->Identifier);
-					}
-				}
-				break;
-
-			case EJson::Array:
-				{
-					if (Element->bHasBeenProcessed)
-					{
-						Writer.WriteArrayEnd();
-					}
-					else
-					{
-						Element->bHasBeenProcessed = true;
-						ElementStack.Push(Element);
-
-						if (bWriteValueOnly)
-						{
-							Writer.WriteArrayStart();
-						}
-						else
-						{
-							Writer.WriteArrayStart(Element->Identifier);
-						}
-
-						TArray<TSharedPtr<FJsonValue>> Values = Element->Value->AsArray();
-
-						for (int Index = Values.Num() - 1; Index >= 0; --Index)
-						{
-							ElementStack.Push(MakeShared<FElement>(Values[Index]));
-						}
-					}
-				}
-				break;
-
-			case EJson::Object:
-				{
-					if (Element->bHasBeenProcessed)
-					{
-						Writer.WriteObjectEnd();
-					}
-					else
-					{
-						Element->bHasBeenProcessed = true;
-						ElementStack.Push(Element);
-
-						if (bWriteValueOnly)
-						{
-							Writer.WriteObjectStart();
-						}
-						else
-						{
-							Writer.WriteObjectStart(Element->Identifier);
-						}
-
-						TArray<FString> Keys; 
-						TArray<TSharedPtr<FJsonValue>> Values;
-						TSharedPtr<FJsonObject> ElementObject = Element->Value->AsObject();
-						ElementObject->Values.GenerateKeyArray(Keys);
-						ElementObject->Values.GenerateValueArray(Values);
-
-						check(Keys.Num() == Values.Num());
-
-						for (int Index = Values.Num() - 1; Index >= 0; --Index)
-						{
-							ElementStack.Push(MakeShared<FElement>(Keys[Index], Values[Index]));
-						}
-					}
-				}
-				break;
-
-			case EJson::None:
-			default: 
-				UE_LOG(LogJson, Fatal,TEXT("Could not print Json Value, unrecognized type."));
+				continue;
 			}
+			
+			if (Policy::SerializeIfNumber(Element, Writer, bWriteValueOnly))
+			{
+				continue;
+			}
+			
+			if (Policy::SerializeIfString(Element, Writer, bWriteValueOnly))
+			{
+				continue;
+			}
+			
+			if (Policy::SerializeIfNull(Element, Writer, bWriteValueOnly))
+			{
+				continue;
+			}
+			
+			if (Policy::SerializeIfArray(ElementStack, Element, Writer, bWriteValueOnly))
+			{
+				continue;
+			}
+			
+			if (Policy::SerializeIfObject(ElementStack, Element, Writer, bWriteValueOnly))
+			{
+				continue;
+			}
+			
+			UE_LOG(LogJson, Fatal, TEXT("Could not print Json Value, unrecognized type."));
 		}
 
 		if (bCloseWriter)
@@ -569,3 +653,5 @@ private:
 		}
 	}
 };
+
+using FJsonSerializer = TJsonSerializer<FJsonSerializerPolicy_JsonObject>;
