Added missing typename indicators in JsonSerializer

diff --git a/Engine/Source/Runtime/Json/Public/Serialization/JsonSerializer.h b/Engine/Source/Runtime/Json/Public/Serialization/JsonSerializer.h
--- a/Engine/Source/Runtime/Json/Public/Serialization/JsonSerializer.h
+++ b/Engine/Source/Runtime/Json/Public/Serialization/JsonSerializer.h
@@ -280,13 +280,13 @@ public:
 	};
 
 	template <class CharType>
-	static bool Deserialize(const TSharedRef<TJsonReader<CharType>>& Reader, Policy::FValue& OutValue, EFlags InOptions = EFlags::None)
+	static bool Deserialize(const TSharedRef<TJsonReader<CharType>>& Reader, typename Policy::FValue& OutValue, EFlags InOptions = EFlags::None)
 	{
 		return Deserialize(*Reader, OutValue, InOptions);
 	}
 
 	template <class CharType>
-	static bool Deserialize(TJsonReader<CharType>& Reader, Policy::FValue& OutValue, EFlags InOptions = EFlags::None)
+	static bool Deserialize(TJsonReader<CharType>& Reader, typename Policy::FValue& OutValue, EFlags InOptions = EFlags::None)
 	{
 		typename Policy::StackState State;
 		if (!Deserialize(Reader, /*OUT*/State, InOptions))
@@ -298,13 +298,13 @@ public:
 	}
 
 	template <class CharType>
-	static bool Deserialize(const TSharedRef<TJsonReader<CharType>>& Reader, Policy::FArrayOfValues& OutArray, EFlags InOptions = EFlags::None)
+	static bool Deserialize(const TSharedRef<TJsonReader<CharType>>& Reader, typename Policy::FArrayOfValues& OutArray, EFlags InOptions = EFlags::None)
 	{
 		return Deserialize(*Reader, OutArray, InOptions);
 	}
 
 	template <class CharType>
-	static bool Deserialize(TJsonReader<CharType>& Reader, Policy::FArrayOfValues& OutArray, EFlags InOptions = EFlags::None)
+	static bool Deserialize(TJsonReader<CharType>& Reader, typename Policy::FArrayOfValues& OutArray, EFlags InOptions = EFlags::None)
 	{
 		typename Policy::StackState State;
 		if (!Deserialize(Reader, /*OUT*/State, InOptions))
@@ -316,13 +316,13 @@ public:
 	}
 
 	template <class CharType>
-	static bool Deserialize(const TSharedRef<TJsonReader<CharType>>& Reader, Policy::FMapOfValues& OutMap, EFlags InOptions = EFlags::None)
+	static bool Deserialize(const TSharedRef<TJsonReader<CharType>>& Reader, typename Policy::FMapOfValues& OutMap, EFlags InOptions = EFlags::None)
 	{
 		return Deserialize(*Reader, OutMap, InOptions);
 	}
 
 	template <class CharType>
-	static bool Deserialize(TJsonReader<CharType>& Reader, Policy::FMapOfValues& OutMap, EFlags InOptions = EFlags::None)
+	static bool Deserialize(TJsonReader<CharType>& Reader, typename Policy::FMapOfValues& OutMap, EFlags InOptions = EFlags::None)
 	{
 		typename Policy::StackState State;
 		if (!Deserialize(Reader, /*OUT*/State, InOptions))
@@ -356,7 +356,7 @@ public:
 	 * @return				Returns true if the serialization was successful, false otherwise.
 	 */
 	template <class CharType, class PrintPolicy>
-	static bool Serialize(const Policy::FValue& Value, const FString& Identifier, const TSharedRef<TJsonWriter<CharType, PrintPolicy>>& Writer, bool bCloseWriter = true)
+	static bool Serialize(const typename Policy::FValue& Value, const FString& Identifier, const TSharedRef<TJsonWriter<CharType, PrintPolicy>>& Writer, bool bCloseWriter = true)
 	{
 		return Serialize(Value, Identifier, *Writer, bCloseWriter);
 	}
@@ -384,7 +384,7 @@ public:
 	 * @return				Returns true if the serialization was successful, false otherwise.
 	 */
 	template <class CharType, class PrintPolicy>
-	static bool Serialize(const Policy::FValue& Value, const FString& Identifier, TJsonWriter<CharType, PrintPolicy>& Writer, bool bCloseWriter = true)
+	static bool Serialize(const typename Policy::FValue& Value, const FString& Identifier, TJsonWriter<CharType, PrintPolicy>& Writer, bool bCloseWriter = true)
 	{
 		const TSharedRef<typename Policy::FElement> StartingElement = MakeShared<typename Policy::FElement>(Identifier, Value);
 		return Serialize<CharType, PrintPolicy>(StartingElement, Writer, bCloseWriter);
@@ -405,7 +405,7 @@ public:
 	 * @return				Returns true if the serialization was successful, false otherwise.
 	 */
 	template <class CharType, class PrintPolicy>
-	static bool Serialize(const Policy::FArrayOfValues& Array, const TSharedRef<TJsonWriter<CharType, PrintPolicy>>& Writer, bool bCloseWriter = true)
+	static bool Serialize(const typename Policy::FArrayOfValues& Array, const TSharedRef<TJsonWriter<CharType, PrintPolicy>>& Writer, bool bCloseWriter = true)
 	{
 		return Serialize(Array, *Writer, bCloseWriter);
 	}
@@ -425,7 +425,7 @@ public:
 	 * @return				Returns true if the serialization was successful, false otherwise.
 	 */
 	template <class CharType, class PrintPolicy>
-	static bool Serialize(const Policy::FArrayOfValues& Array, TJsonWriter<CharType, PrintPolicy>& Writer, bool bCloseWriter = true )
+	static bool Serialize(const typename Policy::FArrayOfValues& Array, TJsonWriter<CharType, PrintPolicy>& Writer, bool bCloseWriter = true )
 	{
 		const TSharedRef<typename Policy::FElement> StartingElement = MakeShared<typename Policy::FElement>(Array);
 		return Serialize<CharType, PrintPolicy>(StartingElement, Writer, bCloseWriter);
@@ -446,7 +446,7 @@ public:
 	 * @return				Returns true if the serialization was successful, false otherwise.
 	 */
 	template <class CharType, class PrintPolicy>
-	static bool Serialize(const Policy::FMapOfValues& Object, const TSharedRef<TJsonWriter<CharType, PrintPolicy>>& Writer, bool bCloseWriter = true )
+	static bool Serialize(const typename Policy::FMapOfValues& Object, const TSharedRef<TJsonWriter<CharType, PrintPolicy>>& Writer, bool bCloseWriter = true )
 	{
 		return Serialize(Object, *Writer, bCloseWriter);
 	}
@@ -466,7 +466,7 @@ public:
 	 * @return				Returns true if the serialization was successful, false otherwise.
 	 */
 	template <class CharType, class PrintPolicy>
-	static bool Serialize(const Policy::FMapOfValues& Object, TJsonWriter<CharType, PrintPolicy>& Writer, bool bCloseWriter = true)
+	static bool Serialize(const typename Policy::FMapOfValues& Object, TJsonWriter<CharType, PrintPolicy>& Writer, bool bCloseWriter = true)
 	{
 		const TSharedRef<typename Policy::FElement> StartingElement = MakeShared<typename Policy::FElement>(Object);
 		return Serialize<CharType, PrintPolicy>(StartingElement, Writer, bCloseWriter);
@@ -474,7 +474,7 @@ public:
 
 private:
 	template <class CharType>
-	static bool Deserialize(TJsonReader<CharType>& Reader, Policy::StackState& OutStackState, EFlags InOptions)
+	static bool Deserialize(TJsonReader<CharType>& Reader, typename Policy::StackState& OutStackState, EFlags InOptions)
 	{
 		TArray<TSharedRef<typename Policy::StackState>> ScopeStack; 
 		TSharedPtr<typename Policy::StackState> CurrentState;
