Fix for crash in UE::Json::ParseInPlace with an empty input

diff --git a/Engine/Source/Runtime/Json/Private/JsonUtils/RapidJsonUtils.cpp b/Engine/Source/Runtime/Json/Private/JsonUtils/RapidJsonUtils.cpp
--- a/Engine/Source/Runtime/Json/Private/JsonUtils/RapidJsonUtils.cpp
+++ b/Engine/Source/Runtime/Json/Private/JsonUtils/RapidJsonUtils.cpp
@@ -248,6 +248,15 @@ TValueOrError<FDocument, FParseError> Parse(const FStringView JsonText)
 
 TValueOrError<FDocument, FParseError> ParseInPlace(TArrayView<TCHAR> JsonText)
 {
+	if (JsonText.IsEmpty())
+	{
+		// the in-place read crashes if the text is empty
+		// the non in-place parse returns this error code so this just makes it work the same
+		return MakeError(FParseError{
+			.ErrorCode = rapidjson::kParseErrorDocumentEmpty
+		});
+	}
+
 	if (!JsonText.IsEmpty() && JsonText[JsonText.Num()-1] != 0)
 	{
 		return MakeError(FParseError{
diff --git a/Engine/Source/Runtime/Json/Private/Tests/RapidJsonUtilsTests.cpp b/Engine/Source/Runtime/Json/Private/Tests/RapidJsonUtilsTests.cpp
--- a/Engine/Source/Runtime/Json/Private/Tests/RapidJsonUtilsTests.cpp
+++ b/Engine/Source/Runtime/Json/Private/Tests/RapidJsonUtilsTests.cpp
@@ -182,4 +182,20 @@ TEST_CASE_NAMED(FRapidJsonUtilsErrorHandlingTests, "System::Engine::FileSystem::
 
 }
 
+/** Check to make sure it doesn't crash with an empty input **/
+TEST_CASE_NAMED(FRapidJsonUtilsParseEmpty, "System::Engine::FileSystem::RapidJSONUtils::ParseEmpty", "[ApplicationContextMask][SmokeFilter]")
+{
+	using namespace UE::Json;
+
+	TValueOrError<FDocument, FParseError> Result = Parse({});
+
+	REQUIRE(Result.HasError());
+	REQUIRE(Result.GetError().ErrorCode == rapidjson::kParseErrorDocumentEmpty);
+
+	TValueOrError<FDocument, FParseError> InPlaceResult = ParseInPlace({});
+	REQUIRE(InPlaceResult.HasError());
+	REQUIRE(InPlaceResult.GetError().ErrorCode == rapidjson::kParseErrorDocumentEmpty);
+}
+
+
 #endif // WITH_TESTS
