Fixed prefixes for several internal JSON serializer types.

#rnx

diff --git a/Engine/Source/Runtime/Json/Public/Serialization/JsonSerializable.h b/Engine/Source/Runtime/Json/Public/Serialization/JsonSerializable.h
--- a/Engine/Source/Runtime/Json/Public/Serialization/JsonSerializable.h
+++ b/Engine/Source/Runtime/Json/Public/Serialization/JsonSerializable.h
@@ -135,16 +135,16 @@ inline bool FromJson(TArray<T>& OutArray, TStringView<CharType> JsonString)
 	return false;
 }
 
-using ReturnStringArgs = TTuple<FString* /*OutValue*/, bool /*bPrettyPrint*/>;
+using FReturnStringArgs = TTuple<FString* /*OutValue*/, bool /*bPrettyPrint*/>;
 
-using PrettyWriter = TSharedRef<TJsonWriter<>>;
-using CondensedWriter = TSharedRef<TJsonWriter<TCHAR, TCondensedJsonPrintPolicy<TCHAR>>>;
-using WriterVariants = TVariant<PrettyWriter, CondensedWriter>;
+using FPrettyWriter = TSharedRef<TJsonWriter<>>;
+using FCondensedWriter = TSharedRef<TJsonWriter<TCHAR, TCondensedJsonPrintPolicy<TCHAR>>>;
+using FWriterVariants = TVariant<FPrettyWriter, FCondensedWriter>;
 
-using ToJsonVariantArgs = TVariant<ReturnStringArgs, WriterVariants>;
+using FToJsonVariantArgs = TVariant<FReturnStringArgs, FWriterVariants>;
 
-using PrettySerializer = FJsonSerializerWriter<>;
-using CondensedSerializer = FJsonSerializerWriter<TCHAR, TCondensedJsonPrintPolicy<TCHAR>>;
+using FPrettySerializer = FJsonSerializerWriter<>;
+using FCondensedSerializer = FJsonSerializerWriter<TCHAR, TCondensedJsonPrintPolicy<TCHAR>>;
 
 template<typename T, typename...SerializerArgsT>
 inline void ToJson_SerializeArrayElements(TArray<T>& InArray, SerializerArgsT...Args)
@@ -165,27 +165,27 @@ inline void ToJson_SerializeArrayElements(TArray<T*>& InArray, SerializerArgsT..
 }
 
 template<typename T>
-inline void ToJson(TArray<T>& InArray, const ToJsonVariantArgs& InArgs)
+inline void ToJson(TArray<T>& InArray, const FToJsonVariantArgs& InArgs)
 {
-	using PrettySerializerAndWriter = TTuple<PrettySerializer, PrettyWriter>;
-	using CondensedSerializerAndWriter = TTuple<CondensedSerializer, CondensedWriter>;
+	using FPrettySerializerAndWriter = TTuple<FPrettySerializer, FPrettyWriter>;
+	using FCondensedSerializerAndWriter = TTuple<FCondensedSerializer, FCondensedWriter>;
 
-	using SerializerVariant = TVariant<PrettySerializerAndWriter, CondensedSerializerAndWriter>;
+	using SerializerVariant = TVariant<FPrettySerializerAndWriter, FCondensedSerializerAndWriter>;
 
 	SerializerVariant SerializerToUse = ::Visit([](auto& StoredValue)
 		{
 			using StoredValueType = std::decay_t<decltype(StoredValue)>;
-			if constexpr (std::is_same_v<StoredValueType, ReturnStringArgs>)
+			if constexpr (std::is_same_v<StoredValueType, FReturnStringArgs>)
 			{
 				if (StoredValue.template Get<1>())
 				{
-					PrettyWriter NewWriter = TJsonWriterFactory<>::Create(StoredValue.template Get<0>());
-					return SerializerVariant(TInPlaceType<PrettySerializerAndWriter>(), PrettySerializer(NewWriter), NewWriter);
+					FPrettyWriter NewWriter = TJsonWriterFactory<>::Create(StoredValue.template Get<0>());
+					return SerializerVariant(TInPlaceType<FPrettySerializerAndWriter>(), FPrettySerializer(NewWriter), NewWriter);
 				}
 				else
 				{
-					CondensedWriter NewWriter = TJsonWriterFactory<TCHAR, TCondensedJsonPrintPolicy<TCHAR>>::Create(StoredValue.template Get<0>());
-					return SerializerVariant(TInPlaceType<CondensedSerializerAndWriter>(), CondensedSerializer(NewWriter), NewWriter);
+					FCondensedWriter NewWriter = TJsonWriterFactory<TCHAR, TCondensedJsonPrintPolicy<TCHAR>>::Create(StoredValue.template Get<0>());
+					return SerializerVariant(TInPlaceType<FCondensedSerializerAndWriter>(), FCondensedSerializer(NewWriter), NewWriter);
 				}
 			}
 			else
@@ -193,20 +193,20 @@ inline void ToJson(TArray<T>& InArray, const ToJsonVariantArgs& InArgs)
 				return ::Visit([](auto& StoredWriter)
 					{
 						using StoredWriterType = std::decay_t<decltype(StoredWriter)>;
-						if constexpr (std::is_same_v<StoredWriterType, PrettyWriter>)
+						if constexpr (std::is_same_v<StoredWriterType, FPrettyWriter>)
 						{
-							return SerializerVariant(TInPlaceType<PrettySerializerAndWriter>(), PrettySerializer(StoredWriter), StoredWriter);
+							return SerializerVariant(TInPlaceType<FPrettySerializerAndWriter>(), FPrettySerializer(StoredWriter), StoredWriter);
 						}
-						else if constexpr (std::is_same_v<StoredWriterType, CondensedWriter>)
+						else if constexpr (std::is_same_v<StoredWriterType, FCondensedWriter>)
 						{
-							return SerializerVariant(TInPlaceType<CondensedSerializerAndWriter>(), CondensedSerializer(StoredWriter), StoredWriter);
+							return SerializerVariant(TInPlaceType<FCondensedSerializerAndWriter>(), FCondensedSerializer(StoredWriter), StoredWriter);
 						}
 					}, StoredValue);
 			}
 		}, InArgs);
 
 
-	const bool bCloseWriter = InArgs.IsType<ReturnStringArgs>();
+	const bool bCloseWriter = InArgs.IsType<FReturnStringArgs>();
 
 	::Visit([bCloseWriter, &InArray](auto& StoredSerializer)
 		{
@@ -253,20 +253,20 @@ template<typename T>
 static const FString ToJson(TArray<T>& InArray, const bool bPrettyPrint = true)
 {
 	FString JsonStr;
-	Private::ToJson(InArray, Private::ToJsonVariantArgs(TInPlaceType<Private::ReturnStringArgs>(), &JsonStr, bPrettyPrint));
+	Private::ToJson(InArray, Private::FToJsonVariantArgs(TInPlaceType<Private::FReturnStringArgs>(), &JsonStr, bPrettyPrint));
 	return JsonStr;
 }
 
 template<typename T>
-static void ToJson(TArray<T>& InArray, Private::PrettyWriter& JsonWriter)
+static void ToJson(TArray<T>& InArray, Private::FPrettyWriter& JsonWriter)
 {
-	Private::ToJson(InArray, Private::ToJsonVariantArgs(TInPlaceType<Private::WriterVariants>(), Private::WriterVariants(TInPlaceType<Private::PrettyWriter>(), JsonWriter)));
+	Private::ToJson(InArray, Private::FToJsonVariantArgs(TInPlaceType<Private::FWriterVariants>(), Private::FWriterVariants(TInPlaceType<Private::FPrettyWriter>(), JsonWriter)));
 }
 
 template<typename T>
-static void ToJson(TArray<T>& InArray, Private::CondensedWriter& JsonWriter)
+static void ToJson(TArray<T>& InArray, Private::FCondensedWriter& JsonWriter)
 {
-	Private::ToJson(InArray, Private::ToJsonVariantArgs(TInPlaceType<Private::WriterVariants>(), Private::WriterVariants(TInPlaceType<Private::CondensedWriter>(), JsonWriter)));
+	Private::ToJson(InArray, Private::FToJsonVariantArgs(TInPlaceType<Private::FWriterVariants>(), Private::FWriterVariants(TInPlaceType<Private::FCondensedWriter>(), JsonWriter)));
 }
 
 } // namespace UE::JsonArray
