Added UE::Overload, for combining multiple lambdas/invocables into one object with overloads.
Added some usage.

diff --git a/Engine/Source/Runtime/Core/Public/Templates/Overload.h b/Engine/Source/Runtime/Core/Public/Templates/Overload.h
new file mode 100644
--- /dev/null
+++ b/Engine/Source/Runtime/Core/Public/Templates/Overload.h
@@ -0,0 +1,81 @@
+// Copyright Epic Games, Inc. All Rights Reserved.
+
+#pragma once
+
+/**
+ * UE::Overload allows the combining of multiple invocables into a single object where they are overloaded.
+ *
+ * Example:
+ *
+ * TTuple<int32, FString, float> Tup1 = ...;
+ * TTuple<TArray<int32>, const UObject*, TCHAR> Tup2 = ...;
+ * VisitTupleElements(
+ *     UE::Overload(
+ *         [](int32 Int, const TArray<int32>& Arr)
+ *         {
+ *             // Called with (Tup1.Get<0>(), Tup2.Get<0>())
+ *         },
+ *         [Capture1](const FString& Str, const UObject* Obj)
+ *         {
+ *             // Called with (Tup1.Get<1>(), Tup2.Get<1>())
+ *         },
+ *         [Capture1, Capture2](float Val, TCHAR Ch)
+ *         {
+ *             // Called with (Tup1.Get<2>(), Tup2.Get<2>())
+ *         },
+ *     )
+ * )
+ *
+ * Warning: all of the invocables' captured state is copied into the result.  If the same capture is present
+ * in multiple callables, e.g. Capture1 above, the overload object will contain multiple copies of that capture.
+ * If this is unacceptable, consider writing a bespoke struct with one copy of the capture with multiple
+ * operator() overloads.
+ */
+
+#include <type_traits>
+
+#include "Templates/Projection.h"
+
+namespace UE::Core::Private
+{
+	template <typename... InvocableTypes>
+	struct TOverload : InvocableTypes...
+	{
+		using InvocableTypes::operator()...;
+	};
+
+	template <typename RetType, typename... ArgTypes>
+	struct TOverloadWrapper
+	{
+		RetType (*Callable)(ArgTypes...);
+
+		UE_REWRITE constexpr RetType operator()(ArgTypes... Args) const
+		{
+			return Callable((ArgTypes&&)Args...);
+		}
+	};
+
+	// This wraps function pointers in an object so they can be inherited by TOverload
+	template <typename RetType, typename... ArgTypes>
+	[[nodiscard]] UE_REWRITE constexpr TOverloadWrapper<RetType, ArgTypes...> MakeCallableObject(RetType(*Callable)(ArgTypes...))
+	{
+		return { Callable };
+	}
+
+	template <typename CallableType>
+		requires (!std::is_pointer_v<std::decay_t<CallableType>>)
+	[[nodiscard]] UE_REWRITE constexpr CallableType&& MakeCallableObject(CallableType&& Callable)
+	{
+		return (CallableType&&)Callable;
+	}
+}
+
+namespace UE
+{
+	// Combines a set of invocable objects into one and overloads them.
+	template <typename... InvocableTypes>
+	[[nodiscard]] UE_REWRITE constexpr auto Overload(InvocableTypes&&... Invocables)
+	{
+		return UE::Core::Private::TOverload{ UE::Core::Private::MakeCallableObject(Projection((InvocableTypes&&)Invocables))... };
+	}
+}
diff --git a/Engine/Source/Runtime/Json/Public/Serialization/JsonSerializable.h b/Engine/Source/Runtime/Json/Public/Serialization/JsonSerializable.h
--- a/Engine/Source/Runtime/Json/Public/Serialization/JsonSerializable.h
+++ b/Engine/Source/Runtime/Json/Public/Serialization/JsonSerializable.h
@@ -9,6 +9,7 @@
 #include "Serialization/JsonSerializerReader.h"
 #include "JsonSerializerWriter.h"
 #include "Misc/TVariant.h"
+#include "Templates/Overload.h"
 
 /**
  * Base class for a JSON serializable object
@@ -170,41 +171,36 @@ inline void ToJson(TArray<T>& InArray, const FToJsonVariantArgs& InArgs)
 	using FPrettySerializerAndWriter = TTuple<FPrettySerializer, FPrettyWriter>;
 	using FCondensedSerializerAndWriter = TTuple<FCondensedSerializer, FCondensedWriter>;
 
-	using SerializerVariant = TVariant<FPrettySerializerAndWriter, FCondensedSerializerAndWriter>;
+	using FSerializerVariant = TVariant<FPrettySerializerAndWriter, FCondensedSerializerAndWriter>;
 
-	SerializerVariant SerializerToUse = ::Visit([](auto& StoredValue)
+	FSerializerVariant SerializerToUse = ::Visit(UE::Overload(
+		[](FReturnStringArgs& ReturnStringArgs)
 		{
-			using StoredValueType = std::decay_t<decltype(StoredValue)>;
-			if constexpr (std::is_same_v<StoredValueType, FReturnStringArgs>)
+			if (ReturnStringArgs.template Get<1>())
 			{
-				if (StoredValue.template Get<1>())
-				{
-					FPrettyWriter NewWriter = TJsonWriterFactory<>::Create(StoredValue.template Get<0>());
-					return SerializerVariant(TInPlaceType<FPrettySerializerAndWriter>(), FPrettySerializer(NewWriter), NewWriter);
-				}
-				else
-				{
-					FCondensedWriter NewWriter = TJsonWriterFactory<TCHAR, TCondensedJsonPrintPolicy<TCHAR>>::Create(StoredValue.template Get<0>());
-					return SerializerVariant(TInPlaceType<FCondensedSerializerAndWriter>(), FCondensedSerializer(NewWriter), NewWriter);
-				}
+				FPrettyWriter NewWriter = TJsonWriterFactory<>::Create(ReturnStringArgs.template Get<0>());
+				return FSerializerVariant(TInPlaceType<FPrettySerializerAndWriter>(), FPrettySerializer(NewWriter), NewWriter);
 			}
 			else
 			{
-				return ::Visit([](auto& StoredWriter)
-					{
-						using StoredWriterType = std::decay_t<decltype(StoredWriter)>;
-						if constexpr (std::is_same_v<StoredWriterType, FPrettyWriter>)
-						{
-							return SerializerVariant(TInPlaceType<FPrettySerializerAndWriter>(), FPrettySerializer(StoredWriter), StoredWriter);
-						}
-						else if constexpr (std::is_same_v<StoredWriterType, FCondensedWriter>)
-						{
-							return SerializerVariant(TInPlaceType<FCondensedSerializerAndWriter>(), FCondensedSerializer(StoredWriter), StoredWriter);
-						}
-					}, StoredValue);
+				FCondensedWriter NewWriter = TJsonWriterFactory<TCHAR, TCondensedJsonPrintPolicy<TCHAR>>::Create(ReturnStringArgs.template Get<0>());
+				return FSerializerVariant(TInPlaceType<FCondensedSerializerAndWriter>(), FCondensedSerializer(NewWriter), NewWriter);
 			}
-		}, InArgs);
-
+		},
+		[](FWriterVariants& WriterVariants)
+		{
+			return ::Visit(UE::Overload(
+				[](FPrettyWriter& PrettyWriter)
+				{
+					return FSerializerVariant(TInPlaceType<FPrettySerializerAndWriter>(), FPrettySerializer(PrettyWriter), PrettyWriter);
+				},
+				[](FCondensedWriter& CondensedWriter)
+				{
+					return FSerializerVariant(TInPlaceType<FCondensedSerializerAndWriter>(), FCondensedSerializer(CondensedWriter), CondensedWriter);
+				}
+			), WriterVariants);
+		}
+	), InArgs);
 
 	const bool bCloseWriter = InArgs.IsType<FReturnStringArgs>();
 
