Undo CL 46107039.

#rnx

diff --git a/Engine/Source/Runtime/Core/Public/Templates/Overload.h b/Engine/Source/Runtime/Core/Public/Templates/Overload.h
deleted file mode 100644
diff --git a/Engine/Source/Runtime/Json/Public/Serialization/JsonSerializable.h b/Engine/Source/Runtime/Json/Public/Serialization/JsonSerializable.h
--- a/Engine/Source/Runtime/Json/Public/Serialization/JsonSerializable.h
+++ b/Engine/Source/Runtime/Json/Public/Serialization/JsonSerializable.h
@@ -9,7 +9,6 @@
 #include "Serialization/JsonSerializerReader.h"
 #include "JsonSerializerWriter.h"
 #include "Misc/TVariant.h"
-#include "Templates/Overload.h"
 
 /**
  * Base class for a JSON serializable object
@@ -171,36 +170,41 @@ inline void ToJson(TArray<T>& InArray, const FToJsonVariantArgs& InArgs)
 	using FPrettySerializerAndWriter = TTuple<FPrettySerializer, FPrettyWriter>;
 	using FCondensedSerializerAndWriter = TTuple<FCondensedSerializer, FCondensedWriter>;
 
-	using FSerializerVariant = TVariant<FPrettySerializerAndWriter, FCondensedSerializerAndWriter>;
+	using SerializerVariant = TVariant<FPrettySerializerAndWriter, FCondensedSerializerAndWriter>;
 
-	FSerializerVariant SerializerToUse = ::Visit(UE::Overload(
-		[](FReturnStringArgs& ReturnStringArgs)
+	SerializerVariant SerializerToUse = ::Visit([](auto& StoredValue)
 		{
-			if (ReturnStringArgs.template Get<1>())
+			using StoredValueType = std::decay_t<decltype(StoredValue)>;
+			if constexpr (std::is_same_v<StoredValueType, FReturnStringArgs>)
 			{
-				FPrettyWriter NewWriter = TJsonWriterFactory<>::Create(ReturnStringArgs.template Get<0>());
-				return FSerializerVariant(TInPlaceType<FPrettySerializerAndWriter>(), FPrettySerializer(NewWriter), NewWriter);
+				if (StoredValue.template Get<1>())
+				{
+					FPrettyWriter NewWriter = TJsonWriterFactory<>::Create(StoredValue.template Get<0>());
+					return SerializerVariant(TInPlaceType<FPrettySerializerAndWriter>(), FPrettySerializer(NewWriter), NewWriter);
+				}
+				else
+				{
+					FCondensedWriter NewWriter = TJsonWriterFactory<TCHAR, TCondensedJsonPrintPolicy<TCHAR>>::Create(StoredValue.template Get<0>());
+					return SerializerVariant(TInPlaceType<FCondensedSerializerAndWriter>(), FCondensedSerializer(NewWriter), NewWriter);
+				}
 			}
 			else
 			{
-				FCondensedWriter NewWriter = TJsonWriterFactory<TCHAR, TCondensedJsonPrintPolicy<TCHAR>>::Create(ReturnStringArgs.template Get<0>());
-				return FSerializerVariant(TInPlaceType<FCondensedSerializerAndWriter>(), FCondensedSerializer(NewWriter), NewWriter);
+				return ::Visit([](auto& StoredWriter)
+					{
+						using StoredWriterType = std::decay_t<decltype(StoredWriter)>;
+						if constexpr (std::is_same_v<StoredWriterType, FPrettyWriter>)
+						{
+							return SerializerVariant(TInPlaceType<FPrettySerializerAndWriter>(), FPrettySerializer(StoredWriter), StoredWriter);
+						}
+						else if constexpr (std::is_same_v<StoredWriterType, FCondensedWriter>)
+						{
+							return SerializerVariant(TInPlaceType<FCondensedSerializerAndWriter>(), FCondensedSerializer(StoredWriter), StoredWriter);
+						}
+					}, StoredValue);
 			}
-		},
-		[](FWriterVariants& WriterVariants)
-		{
-			return ::Visit(UE::Overload(
-				[](FPrettyWriter& PrettyWriter)
-				{
-					return FSerializerVariant(TInPlaceType<FPrettySerializerAndWriter>(), FPrettySerializer(PrettyWriter), PrettyWriter);
-				},
-				[](FCondensedWriter& CondensedWriter)
-				{
-					return FSerializerVariant(TInPlaceType<FCondensedSerializerAndWriter>(), FCondensedSerializer(CondensedWriter), CondensedWriter);
-				}
-			), WriterVariants);
-		}
-	), InArgs);
+		}, InArgs);
+
 
 	const bool bCloseWriter = InArgs.IsType<FReturnStringArgs>();
 
