Add json macro to serialize shared ptr of serializable.

diff --git a/Engine/Source/Runtime/Json/Private/Tests/JsonTests.cpp b/Engine/Source/Runtime/Json/Private/Tests/JsonTests.cpp
--- a/Engine/Source/Runtime/Json/Private/Tests/JsonTests.cpp
+++ b/Engine/Source/Runtime/Json/Private/Tests/JsonTests.cpp
@@ -1516,4 +1516,47 @@ TEST_CASE_NAMED(FJsonSerializerMacrosTest, "System::Engine::FileSystem::JSON::Ma
 	CHECK(ToJsonResult == SourceJsonString);
 }
 
+TEST_CASE("System::Engine::FileSystem::JSON::MacroToSerializeSharedPtr", "[ApplicationContextMask][SmokeFilter]")
+{
+	struct FJsonStruct : public FJsonSerializable
+	{
+		struct FSubJsonStruct : public FJsonSerializable
+		{
+			int32 SubVarInt;
+			FString SubVarString;
+
+			FSubJsonStruct(int32 InSubVarInt = 0, FStringView InSubVarString = TEXTVIEW(""))
+				: SubVarInt(InSubVarInt)
+				, SubVarString(InSubVarString)
+			{
+			}
+
+			BEGIN_JSON_SERIALIZER
+				JSON_SERIALIZE("sub_var_int", SubVarInt);
+				JSON_SERIALIZE("sub_var_string", SubVarString);
+			END_JSON_SERIALIZER
+		};
+
+		int32 VarInt = 0;
+		TSharedPtr<FSubJsonStruct> VarSerializable;
+
+		BEGIN_JSON_SERIALIZER
+			JSON_SERIALIZE("var_int", VarInt);
+			JSON_SERIALIZE_SHARED_OBJECT_SERIALIZABLE("var_shared_ptr", VarSerializable);
+		END_JSON_SERIALIZER;
+	};
+
+	const TCHAR* SourceJsonString = TEXT("{\"var_int\":2,\"var_shared_ptr\":{\"sub_var_int\":10,\"sub_var_string\":\"abc\"}}");
+
+	FJsonStruct JsonStructTarget;
+	JsonStructTarget.FromJson(SourceJsonString);
+
+	CHECK(JsonStructTarget.VarInt == 2);
+	CHECK(JsonStructTarget.VarSerializable->SubVarInt == 10);
+	CHECK(JsonStructTarget.VarSerializable->SubVarString == TEXT("abc"));
+
+	FString ToJsonResult = JsonStructTarget.ToJson(/*bPrettyPrint*/false);
+	CHECK(ToJsonResult == SourceJsonString);
+}
+
 #endif // WITH_TESTS
diff --git a/Engine/Source/Runtime/Json/Public/Serialization/JsonSerializerMacros.h b/Engine/Source/Runtime/Json/Public/Serialization/JsonSerializerMacros.h
--- a/Engine/Source/Runtime/Json/Public/Serialization/JsonSerializerMacros.h
+++ b/Engine/Source/Runtime/Json/Public/Serialization/JsonSerializerMacros.h
@@ -361,6 +361,30 @@
 			} \
 		}
 
+#define JSON_SERIALIZE_SHARED_OBJECT_SERIALIZABLE(JsonName, JsonSerializableObject) \
+		if (Serializer.IsLoading()) \
+		{ \
+			using ObjectType = TRemoveReference<decltype(*JsonSerializableObject)>::Type; \
+			if (Serializer.GetObject()->HasTypedField<EJson::Object>(TEXTVIEW(JsonName))) \
+			{ \
+				TSharedPtr<FJsonObject> JsonObj = Serializer.GetObject()->GetObjectField(TEXTVIEW(JsonName)); \
+				if (JsonObj.IsValid()) \
+				{ \
+					JsonSerializableObject = MakeShared<ObjectType>(); \
+					JsonSerializableObject->FromJson(JsonObj); \
+				} \
+			} \
+		} \
+		else \
+		{ \
+			if (JsonSerializableObject) \
+			{ \
+				Serializer.StartObject(TEXTVIEW(JsonName)); \
+				JsonSerializableObject->Serialize(Serializer, true); \
+				Serializer.EndObject(); \
+			} \
+		}
+
 #define JSON_SERIALIZE_DATETIME_UNIX_TIMESTAMP(JsonName, JsonDateTime) \
 		if (Serializer.IsLoading()) \
 		{ \
