TJsonStringWriter optimization
- There is no need to copy individual characters to the string - we know the character types are identical, and we know exactly how many characters will be appended.

diff --git a/Engine/Source/Runtime/Json/Public/Serialization/JsonWriter.h b/Engine/Source/Runtime/Json/Public/Serialization/JsonWriter.h
--- a/Engine/Source/Runtime/Json/Public/Serialization/JsonWriter.h
+++ b/Engine/Source/Runtime/Json/Public/Serialization/JsonWriter.h
@@ -754,13 +754,9 @@ public:
 
 	virtual bool Close() override
 	{
-		OutString->Reset(Bytes.Num()/sizeof(CharType));
-		for (int32 i = 0; i < Bytes.Num(); i+=sizeof(CharType))
-		{
-			CharType* Char = static_cast<CharType*>(static_cast<void*>(&Bytes[i]));
-			*OutString += *Char;
-		}
-
+		// This makes the slightly questionable assumption that the Bytes array does actually contain a string of contiguous CharType
+		OutString->Reset(Bytes.Num() / sizeof(CharType));
+		OutString->AppendChars((const CharType*)Bytes.GetData(), Bytes.Num() / sizeof(CharType));
 		return TJsonWriter<CharType, PrintPolicy>::Close();
 	}
 
