Fix compilation when UE_DEFINE_LEGACY_MATH_CONSTANT_MACRO_NAMES is 0.

#rnx

#Synced-CL 49029890

diff --git a/Engine/Source/Runtime/AssetRegistry/Private/AssetRegistry.cpp b/Engine/Source/Runtime/AssetRegistry/Private/AssetRegistry.cpp
--- a/Engine/Source/Runtime/AssetRegistry/Private/AssetRegistry.cpp
+++ b/Engine/Source/Runtime/AssetRegistry/Private/AssetRegistry.cpp
@@ -4386,7 +4386,7 @@ float GetAssetAvailabilityProgress(const FAssetData& AssetData, EAssetAvailabili
 
 	bool IsPercentageComplete = (ChunkReportType == EChunkProgressReportingType::PercentageComplete) ? true : false;
 
-	float BestProgress = MAX_FLT;
+	float BestProgress = UE_MAX_FLT;
 
 	// check all chunks to see which has the best time remaining
 	for (int32 PakchunkID : AssetData.GetChunkIDs())
diff --git a/Engine/Source/Runtime/Json/Private/Tests/JsonTests.cpp b/Engine/Source/Runtime/Json/Private/Tests/JsonTests.cpp
--- a/Engine/Source/Runtime/Json/Private/Tests/JsonTests.cpp
+++ b/Engine/Source/Runtime/Json/Private/Tests/JsonTests.cpp
@@ -1176,7 +1176,7 @@ TEST_CASE_NAMED(FJsonAutomationTest, "System::Engine::FileSystem::JSON", "[Appli
 		REQUIRE(Array[3]->IsNull());
 		REQUIRE(Array[4]->AsBool());
 		REQUIRE(!Array[5]->AsBool());
-		REQUIRE(FMath::Abs(Array[6]->AsNumber() - 45.f) < KINDA_SMALL_NUMBER);
+		REQUIRE(FMath::Abs(Array[6]->AsNumber() - 45.f) < UE_KINDA_SMALL_NUMBER);
 		const TSharedPtr<FJsonObject> InnerObject = Array[7]->AsObject();
 		REQUIRE(InnerObject.IsValid());
 
