Remove requirement for JSON parsed in place to be null-terminated.

#tests FoundationTests

#Synced-CL 50925863

diff --git a/Engine/Source/Runtime/Json/Private/JsonUtils/RapidJsonUtils.cpp b/Engine/Source/Runtime/Json/Private/JsonUtils/RapidJsonUtils.cpp
--- a/Engine/Source/Runtime/Json/Private/JsonUtils/RapidJsonUtils.cpp
+++ b/Engine/Source/Runtime/Json/Private/JsonUtils/RapidJsonUtils.cpp
@@ -248,28 +248,70 @@ TValueOrError<FDocument, FParseError> Parse(const FStringView JsonText)
 	return MakeValue(MoveTemp(Result));
 }
 
-TValueOrError<FDocument, FParseError> ParseInPlace(TArrayView<TCHAR> JsonText)
+struct FInSituStream
 {
-	TRACE_CPUPROFILER_EVENT_SCOPE_STR("UE::Json::ParseInPlace");
-	if (JsonText.IsEmpty())
+	using Ch = TCHAR;
+
+	FInSituStream(TArrayView<TCHAR> InJsonText)
+		: JsonText(InJsonText)
 	{
-		// the in-place read crashes if the text is empty
-		// the non in-place parse returns this error code so this just makes it work the same
-		return MakeError(FParseError{
-			.ErrorCode = rapidjson::kParseErrorDocumentEmpty
-		});
 	}
 
-	if (!JsonText.IsEmpty() && JsonText[JsonText.Num()-1] != 0)
+	TCHAR Peek() const
 	{
-		return MakeError(FParseError{
-			.Offset = InternalError_NotNullTerminated
-		});
+		if (ReadCursor < JsonText.Num())
+		{
+			return JsonText[ReadCursor];
+		}
+		return '\0';
 	}
 
-	FDocument Result;
+	TCHAR Take() 
+	{
+		if (ReadCursor < JsonText.Num())
+		{
+			return JsonText[ReadCursor++];
+		}
+		return '\0';
+	}
+
+	TCHAR* PutBegin()
+	{
+		WriteCursor = ReadCursor;
+		check (!JsonText.IsEmpty() && WriteCursor < JsonText.Num());
+		return &JsonText[WriteCursor];
+	}
 
-	rapidjson::GenericInsituStringStream<FDocumentEncoding> Stream(JsonText.GetData());
+	void Put(TCHAR C)
+	{
+		if (WriteCursor < JsonText.Num())
+		{
+			JsonText[WriteCursor++] = C;
+		}
+	}
+
+	size_t PutEnd(TCHAR* BeganAt)
+	{
+		check (!JsonText.IsEmpty() && WriteCursor <= JsonText.Num());
+		return &JsonText[WriteCursor] - BeganAt;
+	}
+
+	size_t Tell() const
+	{
+		return ReadCursor;
+	}
+	
+	TArrayView<TCHAR> JsonText;
+	int32 ReadCursor = 0;
+	int32 WriteCursor = 0;
+};
+
+TValueOrError<FDocument, FParseError> ParseInPlace(TArrayView<TCHAR> JsonText)
+{
+	TRACE_CPUPROFILER_EVENT_SCOPE_STR("UE::Json::ParseInPlace");
+
+	FDocument Result;
+	FInSituStream Stream(JsonText);
 	Result.ParseStream<DefaultParseFlags | rapidjson::kParseInsituFlag>(Stream);
 
 	if (Result.HasParseError())
diff --git a/Engine/Source/Runtime/Json/Private/Tests/RapidJsonUtilsTests.cpp b/Engine/Source/Runtime/Json/Private/Tests/RapidJsonUtilsTests.cpp
--- a/Engine/Source/Runtime/Json/Private/Tests/RapidJsonUtilsTests.cpp
+++ b/Engine/Source/Runtime/Json/Private/Tests/RapidJsonUtilsTests.cpp
@@ -1,5 +1,6 @@
 // Copyright Epic Games, Inc. All Rights Reserved.
 
+#include "Containers/ArrayView.h"
 #if WITH_TESTS
 
 #include "CoreMinimal.h"
@@ -197,5 +198,25 @@ TEST_CASE_NAMED(FRapidJsonUtilsParseEmpty, "System::Engine::FileSystem::RapidJSO
 	REQUIRE(InPlaceResult.GetError().ErrorCode == rapidjson::kParseErrorDocumentEmpty);
 }
 
+/** Check to make inputs aren't required to be null terminated **/
+TEST_CASE_NAMED(FRapidJsonUtilsParseWithoutNullTerminator, "System::Engine::FileSystem::RapidJSONUtils::NoNullTerminator", "[ApplicationContextMask][SmokeFilter]")
+{
+	using namespace UE::Json;
+
+	FString UnterminatedJsonText(TEXT(R"json(
+		{ "key" : "value" };trailing
+		)json"));
+
+	int32 Index = 0;
+	REQUIRE(UnterminatedJsonText.FindLastChar('}', Index));
+	TArrayView<TCHAR> View = MakeArrayView(GetData(UnterminatedJsonText), Index+1);
+
+	TValueOrError<FDocument, FParseError> Result = Parse(FStringView(View));
+	REQUIRE(!Result.HasError());
+
+	TValueOrError<FDocument, FParseError> InPlaceResult = ParseInPlace(MakeArrayView(View));
+	REQUIRE(!InPlaceResult.HasError());
+}
+
 
 #endif // WITH_TESTS
