Removed deprecated portion of FJsonObject's interface

#rnx

#Synced-CL 51716849

diff --git a/Engine/Source/Runtime/Json/Public/Dom/JsonObject.h b/Engine/Source/Runtime/Json/Public/Dom/JsonObject.h
--- a/Engine/Source/Runtime/Json/Public/Dom/JsonObject.h
+++ b/Engine/Source/Runtime/Json/Public/Dom/JsonObject.h
@@ -368,174 +368,4 @@ public:
 public:
 	static JSON_API void Duplicate(const TSharedPtr<const FJsonObject>& Source, const TSharedPtr<FJsonObject>& Dest);
 	static JSON_API void Duplicate(const TSharedPtr<FJsonObject>& Source, TSharedPtr<FJsonObject>& Dest);
-
-#if !PLATFORM_TCHAR_IS_UTF8CHAR
-
-	template <EJson JsonType>
-	UE_DEPRECATED(5.4, "Passing an ANSI string to GetField has been deprecated outside of UTF-8 mode. Please use the overload that takes a TCHAR string.")
-	TSharedPtr<FJsonValue> GetField(FAnsiStringView FieldName) const
-	{
-		return GetField<JsonType>(StringCast<TCHAR>(FieldName.GetData(), FieldName.Len()));
-	}
-
-	UE_DEPRECATED(5.4, "Passing an ANSI string to TryGetField has been deprecated outside of UTF-8 mode. Please use the overload that takes a TCHAR string.")
-	TSharedPtr<FJsonValue> TryGetField(FAnsiStringView FieldName) const
-	{
-		return TryGetField(StringCast<TCHAR>(FieldName.GetData(), FieldName.Len()));
-	}
-
-	UE_DEPRECATED(5.4, "Passing an ANSI string to HasField has been deprecated outside of UTF-8 mode. Please use the overload that takes a TCHAR string.")
-	bool HasField(FAnsiStringView FieldName) const
-	{
-		return HasField(StringCast<TCHAR>(FieldName.GetData(), FieldName.Len()));
-	}
-
-	template <EJson JsonType>
-	UE_DEPRECATED(5.4, "Passing an ANSI string to HasTypedField has been deprecated outside of UTF-8 mode. Please use the overload that takes a TCHAR string.")
-	bool HasTypedField(FAnsiStringView FieldName) const
-	{
-		return HasTypedField<JsonType>(StringCast<TCHAR>(FieldName.GetData(), FieldName.Len()));
-	}
-
-	UE_DEPRECATED(5.4, "Passing an ANSI string to RemoveField has been deprecated outside of UTF-8 mode. Please use the overload that takes a TCHAR string.")
-	void RemoveField(FAnsiStringView FieldName)
-	{
-		return RemoveField(StringCast<TCHAR>(FieldName.GetData(), FieldName.Len()));
-	}
-
-	UE_DEPRECATED(5.4, "Passing an ANSI string to GetNumberField has been deprecated outside of UTF-8 mode. Please use the overload that takes a TCHAR string.")
-	double GetNumberField(FAnsiStringView FieldName) const
-	{
-		return GetNumberField(StringCast<TCHAR>(FieldName.GetData(), FieldName.Len()));
-	}
-
-	UE_DEPRECATED(5.4, "Passing an ANSI string to GetIntegerField has been deprecated outside of UTF-8 mode. Please use the overload that takes a TCHAR string.")
-	int32 GetIntegerField(FAnsiStringView FieldName) const
-	{
-		return GetIntegerField(StringCast<TCHAR>(FieldName.GetData(), FieldName.Len()));
-	}
-
-	UE_DEPRECATED(5.4, "Passing an ANSI string to TryGetNumberField has been deprecated outside of UTF-8 mode. Please use the overload that takes a TCHAR string.")
-	bool TryGetNumberField(FAnsiStringView FieldName, float& OutNumber) const
-	{
-		return TryGetNumberField(StringCast<TCHAR>(FieldName.GetData(), FieldName.Len()), OutNumber);
-	}
-
-	UE_DEPRECATED(5.4, "Passing an ANSI string to TryGetNumberField has been deprecated outside of UTF-8 mode. Please use the overload that takes a TCHAR string.")
-	bool TryGetNumberField(FAnsiStringView FieldName, double& OutNumber) const
-	{
-		return TryGetNumberField(StringCast<TCHAR>(FieldName.GetData(), FieldName.Len()), OutNumber);
-	}
-
-	UE_DEPRECATED(5.4, "Passing an ANSI string to TryGetNumberField has been deprecated outside of UTF-8 mode. Please use the overload that takes a TCHAR string.")
-	bool TryGetNumberField(FAnsiStringView FieldName, int8& OutNumber) const
-	{
-		return TryGetNumberField(StringCast<TCHAR>(FieldName.GetData(), FieldName.Len()), OutNumber);
-	}
-
-	UE_DEPRECATED(5.4, "Passing an ANSI string to TryGetNumberField has been deprecated outside of UTF-8 mode. Please use the overload that takes a TCHAR string.")
-	bool TryGetNumberField(FAnsiStringView FieldName, int16& OutNumber) const
-	{
-		return TryGetNumberField(StringCast<TCHAR>(FieldName.GetData(), FieldName.Len()), OutNumber);
-	}
-
-	UE_DEPRECATED(5.4, "Passing an ANSI string to TryGetNumberField has been deprecated outside of UTF-8 mode. Please use the overload that takes a TCHAR string.")
-	bool TryGetNumberField(FAnsiStringView FieldName, int32& OutNumber) const
-	{
-		return TryGetNumberField(StringCast<TCHAR>(FieldName.GetData(), FieldName.Len()), OutNumber);
-	}
-
-	UE_DEPRECATED(5.4, "Passing an ANSI string to TryGetNumberField has been deprecated outside of UTF-8 mode. Please use the overload that takes a TCHAR string.")
-	bool TryGetNumberField(FAnsiStringView FieldName, int64& OutNumber) const
-	{
-		return TryGetNumberField(StringCast<TCHAR>(FieldName.GetData(), FieldName.Len()), OutNumber);
-	}
-
-	UE_DEPRECATED(5.4, "Passing an ANSI string to TryGetNumberField has been deprecated outside of UTF-8 mode. Please use the overload that takes a TCHAR string.")
-	bool TryGetNumberField(FAnsiStringView FieldName, uint8& OutNumber) const
-	{
-		return TryGetNumberField(StringCast<TCHAR>(FieldName.GetData(), FieldName.Len()), OutNumber);
-	}
-
-	UE_DEPRECATED(5.4, "Passing an ANSI string to TryGetNumberField has been deprecated outside of UTF-8 mode. Please use the overload that takes a TCHAR string.")
-	bool TryGetNumberField(FAnsiStringView FieldName, uint16& OutNumber) const
-	{
-		return TryGetNumberField(StringCast<TCHAR>(FieldName.GetData(), FieldName.Len()), OutNumber);
-	}
-
-	UE_DEPRECATED(5.4, "Passing an ANSI string to TryGetNumberField has been deprecated outside of UTF-8 mode. Please use the overload that takes a TCHAR string.")
-	bool TryGetNumberField(FAnsiStringView FieldName, uint32& OutNumber) const
-	{
-		return TryGetNumberField(StringCast<TCHAR>(FieldName.GetData(), FieldName.Len()), OutNumber);
-	}
-
-	UE_DEPRECATED(5.4, "Passing an ANSI string to TryGetNumberField has been deprecated outside of UTF-8 mode. Please use the overload that takes a TCHAR string.")
-	bool TryGetNumberField(FAnsiStringView FieldName, uint64& OutNumber) const
-	{
-		return TryGetNumberField(StringCast<TCHAR>(FieldName.GetData(), FieldName.Len()), OutNumber);
-	}
-
-	UE_DEPRECATED(5.4, "Passing an ANSI string to GetStringField has been deprecated outside of UTF-8 mode. Please use the overload that takes a TCHAR string.")
-	FString GetStringField(FAnsiStringView FieldName) const
-	{
-		return GetStringField(StringCast<TCHAR>(FieldName.GetData(), FieldName.Len()));
-	}
-
-	UE_DEPRECATED(5.4, "Passing an ANSI string to TryGetStringField has been deprecated outside of UTF-8 mode. Please use the overload that takes a TCHAR string.")
-	bool TryGetStringField(FAnsiStringView FieldName, FString& OutString) const
-	{
-		return TryGetStringField(StringCast<TCHAR>(FieldName.GetData(), FieldName.Len()), OutString);
-	}
-
-	UE_DEPRECATED(5.4, "Passing an ANSI string to TryGetStringArrayField has been deprecated outside of UTF-8 mode. Please use the overload that takes a TCHAR string.")
-	bool TryGetStringArrayField(FAnsiStringView FieldName, TArray<FString>& OutArray) const
-	{
-		return TryGetStringArrayField(StringCast<TCHAR>(FieldName.GetData(), FieldName.Len()), OutArray);
-	}
-
-	template <typename TEnum>
-	UE_DEPRECATED(5.4, "Passing an ANSI string to TryGetEnumArrayField has been deprecated outside of UTF-8 mode. Please use the overload that takes a TCHAR string.")
-	bool TryGetEnumArrayField(FAnsiStringView FieldName, TArray<TEnum>& OutArray) const
-	{
-		return TryGetEnumArrayField(StringCast<TCHAR>(FieldName.GetData(), FieldName.Len()), OutArray);
-	}
-
-	UE_DEPRECATED(5.4, "Passing an ANSI string to GetBoolField has been deprecated outside of UTF-8 mode. Please use the overload that takes a TCHAR string.")
-	bool GetBoolField(FAnsiStringView FieldName) const
-	{
-		return GetBoolField(StringCast<TCHAR>(FieldName.GetData(), FieldName.Len()));
-	}
-
-	UE_DEPRECATED(5.4, "Passing an ANSI string to TryGetBoolField has been deprecated outside of UTF-8 mode. Please use the overload that takes a TCHAR string.")
-	bool TryGetBoolField(FAnsiStringView FieldName, bool& OutBool) const
-	{
-		return TryGetBoolField(StringCast<TCHAR>(FieldName.GetData(), FieldName.Len()), OutBool);
-	}
-
-	UE_DEPRECATED(5.4, "Passing an ANSI string to GetArrayField has been deprecated outside of UTF-8 mode. Please use the overload that takes a TCHAR string.")
-	const TArray<TSharedPtr<FJsonValue>>& GetArrayField(FAnsiStringView FieldName) const
-	{
-		return GetArrayField(StringCast<TCHAR>(FieldName.GetData(), FieldName.Len()));
-	}
-
-	UE_DEPRECATED(5.4, "Passing an ANSI string to TryGetArrayField has been deprecated outside of UTF-8 mode. Please use the overload that takes a TCHAR string.")
-	bool TryGetArrayField(FAnsiStringView FieldName, const TArray<TSharedPtr<FJsonValue>>*& OutArray) const
-	{
-		return TryGetArrayField(StringCast<TCHAR>(FieldName.GetData(), FieldName.Len()), OutArray);
-	}
-
-	UE_DEPRECATED(5.4, "Passing an ANSI string to GetObjectField has been deprecated outside of UTF-8 mode. Please use the overload that takes a TCHAR string.")
-	const TSharedPtr<FJsonObject>& GetObjectField(FAnsiStringView FieldName) const
-	{
-		return GetObjectField(StringCast<TCHAR>(FieldName.GetData(), FieldName.Len()));
-	}
-
-	UE_DEPRECATED(5.4, "Passing an ANSI string to TryGetObjectField has been deprecated outside of UTF-8 mode. Please use the overload that takes a TCHAR string.")
-	bool TryGetObjectField(FAnsiStringView FieldName, const TSharedPtr<FJsonObject>*& OutObject) const
-	{
-		return TryGetObjectField(StringCast<TCHAR>(FieldName.GetData(), FieldName.Len()), OutObject);
-	}
-
-#endif // !PLATFORM_TCHAR_IS_UTF8CHAR
-
 };
